Extend Unreal Insights to allow tracing protocols accross multiple machines participating in a Multi User session.

New trace macros have been introduced to ConcertTransport/Tracing. Every macro is given a sequence ID, which is passed across multiple machines to correleate data.

Introduces the ConcertInsights folder in which:
- ConcertInsightsVisualizer integrates with Unreal Insights and analyses data generated by the the added tracing macros.
- ConcertInsightsClient adds a status bar entry for starting a synchronized trace: it sends a network message for starting a trace on connected clients and the server.
- ConcertInsightsServer receives the message sent by ConcertInsightsClient
- ConcertInsightsCore contains shared code for the Client and Server plugins

Read the README.txt for instructions on how to set up.

diff --git a/Engine/Plugins/Developer/Concert/ConcertInsights/ConcertInsightsClient/ConcertInsightsClient.uplugin b/Engine/Plugins/Developer/Concert/ConcertInsights/ConcertInsightsClient/ConcertInsightsClient.uplugin
new file mode 100644
--- /dev/null
+++ b/Engine/Plugins/Developer/Concert/ConcertInsights/ConcertInsightsClient/ConcertInsightsClient.uplugin
@@ -0,0 +1,42 @@
+﻿{
+	"FileVersion": 3,
+	"Version": 1,
+	"VersionName": "1.0",
+	"FriendlyName": "ConcertInsightsClient",
+	"Description": "Extends status bar so you can start a synchronized trace on all connected Concert endpoints.",
+	"Category": "Networking",
+	"CreatedBy": "Epic Games, Inc.",
+	"CreatedByURL": "",
+	"DocsURL": "",
+	"MarketplaceURL": "",
+	"CanContainContent": true,
+	"Hidden": true,
+	"IsBetaVersion": false,
+	"IsExperimentalVersion": true,
+	"Installed": false,
+	"Modules": [
+		{
+			"Name": "ConcertInsightsClient",
+			"Type": "Editor",
+			"LoadingPhase": "Default"
+		}
+	],
+	"Plugins": [
+		{
+			"Name": "ConcertInsightsCore",
+			"Enabled": true
+		},
+		{
+			"Name": "ConcertSharedSlate",
+			"Enabled": true
+		},
+		{
+			"Name": "ConcertSyncClient",
+			"Enabled": true
+		},
+		{
+			"Name": "TraceUtilities",
+			"Enabled": true
+		}
+	]
+}
\ No newline at end of file
diff --git a/Engine/Plugins/Developer/Concert/ConcertInsights/ConcertInsightsClient/Source/ConcertInsightsClient/ConcertInsightsClient.Build.cs b/Engine/Plugins/Developer/Concert/ConcertInsights/ConcertInsightsClient/Source/ConcertInsightsClient/ConcertInsightsClient.Build.cs
new file mode 100644
--- /dev/null
+++ b/Engine/Plugins/Developer/Concert/ConcertInsights/ConcertInsightsClient/Source/ConcertInsightsClient/ConcertInsightsClient.Build.cs
@@ -0,0 +1,35 @@
+﻿// Copyright Epic Games, Inc. All Rights Reserved.
+
+using UnrealBuildTool;
+
+public class ConcertInsightsClient : ModuleRules
+{
+    public ConcertInsightsClient(ReadOnlyTargetRules Target) : base(Target)
+    {
+        PCHUsage = ModuleRules.PCHUsageMode.UseExplicitOrSharedPCHs;
+
+        PublicDependencyModuleNames.AddRange(
+            new string[]
+            {
+                "Core",
+            }
+        );
+
+        PrivateDependencyModuleNames.AddRange(
+            new string[]
+            {
+	            "ConcertInsightsCore",
+	            "ConcertSharedSlate",
+	            "ConcertSyncClient",
+                "CoreUObject",
+                "EditorTraceUtilities",
+                "Engine",
+                "Slate",
+                "SlateCore",
+                "ToolMenus", 
+            }
+        );
+        
+        ShortName = "CrtInEd";
+    }
+}
\ No newline at end of file
diff --git a/Engine/Plugins/Developer/Concert/ConcertInsights/ConcertInsightsClient/Source/ConcertInsightsClient/Private/ClientTraceControls.cpp b/Engine/Plugins/Developer/Concert/ConcertInsights/ConcertInsightsClient/Source/ConcertInsightsClient/Private/ClientTraceControls.cpp
new file mode 100644
--- /dev/null
+++ b/Engine/Plugins/Developer/Concert/ConcertInsights/ConcertInsightsClient/Source/ConcertInsightsClient/Private/ClientTraceControls.cpp
@@ -0,0 +1,91 @@
+// Copyright Epic Games, Inc. All Rights Reserved.
+
+#include "ClientTraceControls.h"
+
+#include "ConcertInsightsClientSettings.h"
+#include "IConcertSyncClient.h"
+#include "IConcertSyncClientModule.h"
+
+#include "IEditorTraceUtilitiesModule.h"
+#include "ProfilingDebugging/TraceAuxiliary.h"
+
+namespace UE::ConcertInsightsClient
+{
+	FClientTraceControls::FClientTraceControls()
+	{
+		if (const TSharedPtr<IConcertSyncClient> Client = IConcertSyncClientModule::Get().GetClient(TEXT("MultiUser")))
+		{
+			Client->GetConcertClient()->OnSessionStartup().AddRaw(this, &FClientTraceControls::OnSessionStart);
+			Client->GetConcertClient()->OnSessionShutdown().AddRaw(this, &FClientTraceControls::OnSessionStopped);
+		}
+	}
+
+	FClientTraceControls::~FClientTraceControls()
+	{
+		if (!IConcertSyncClientModule::IsAvailable())
+		{
+			return;
+		}
+		
+		if (const TSharedPtr<IConcertSyncClient> Client = IConcertSyncClientModule::Get().GetClient(TEXT("MultiUser")))
+		{
+			Client->GetConcertClient()->OnSessionStartup().RemoveAll(this);
+		}
+		
+		// ~FTraceControls will clean up delegates of registered sessions itself
+	}
+
+	ConcertInsightsCore::FStartTraceArgs FClientTraceControls::GetDefaultSynchronizedTraceArgs() const
+	{
+		EditorTraceUtilities::IEditorTraceUtilitiesModule& Module = EditorTraceUtilities::IEditorTraceUtilitiesModule::Get();
+		const EditorTraceUtilities::FStatusBarTraceSettings& Settings = Module.GetTraceSettings();
+
+		ConcertInsightsCore::FStartTraceArgs Result;
+		switch (Settings.TraceDestination)
+		{
+		case EditorTraceUtilities::ETraceDestination::TraceStore:
+			Result.ConnectionType = EConcertTraceTargetType::Network;
+			// TODO DP: Figure out the IP address of the trace store since localhost will mean something different for all the other machines
+			Result.Target = UConcertInsightsClientSettings::Get()->SynchronizedTraceDestinationIP;
+			break;
+		case EditorTraceUtilities::ETraceDestination::File:
+			Result.ConnectionType = EConcertTraceTargetType::File;
+			Result.Target = nullptr;
+			break;
+		default:
+			checkNoEntry();
+		}
+		
+		return Result;
+	}
+
+	bool FClientTraceControls::CanSendRequestsToEndpoint(const FGuid& EndpointId, const IConcertSession& Session) const
+	{
+		// Do not send to own endpoint ID - it does not break anything but there's no point.
+		const TSharedPtr<IConcertSyncClient> Client = IConcertSyncClientModule::Get().GetClient(TEXT("MultiUser"));
+		const TSharedPtr<IConcertClientSession> ClientSession = Client->GetConcertClient()->GetCurrentSession();
+		return !ClientSession || EndpointId != ClientSession->GetSessionClientEndpointId();
+	}
+
+	ConcertInsightsCore::FInitArgs FClientTraceControls::GetInitEventArgs() const
+	{
+		const TSharedPtr<IConcertSyncClient> Client = IConcertSyncClientModule::Get().GetClient(TEXT("MultiUser"));
+		const TSharedPtr<IConcertClientSession> Session = Client->GetConcertClient()->GetCurrentSession();
+		if (!Session)
+		{
+			return {};
+		}
+		
+		return { Session->GetSessionClientEndpointId(), Session->GetLocalClientInfo().DisplayName, false };
+	}
+
+	void FClientTraceControls::OnSessionStart(TSharedRef<IConcertClientSession> Session)
+	{
+		RegisterTraceRequestsHandler(Session);
+	}
+
+	void FClientTraceControls::OnSessionStopped(TSharedRef<IConcertClientSession> Session)
+	{
+		OnLeaveSession(*Session);
+	}
+}
diff --git a/Engine/Plugins/Developer/Concert/ConcertInsights/ConcertInsightsClient/Source/ConcertInsightsClient/Private/ClientTraceControls.h b/Engine/Plugins/Developer/Concert/ConcertInsights/ConcertInsightsClient/Source/ConcertInsightsClient/Private/ClientTraceControls.h
new file mode 100644
--- /dev/null
+++ b/Engine/Plugins/Developer/Concert/ConcertInsights/ConcertInsightsClient/Source/ConcertInsightsClient/Private/ClientTraceControls.h
@@ -0,0 +1,36 @@
+// Copyright Epic Games, Inc. All Rights Reserved.
+
+#pragma once
+
+#include "IConcertSyncClient.h"
+#include "TraceControls.h"
+
+namespace UE::ConcertInsightsClient
+{
+	/**
+	 * Client controls for the editor.
+	 * Puts the client's session and display name into the init args.
+	 */
+	class FClientTraceControls : public ConcertInsightsCore::FTraceControls
+	{
+		friend class FTraceControls; // For FTraceControls::Make
+	public:
+		
+		virtual ~FClientTraceControls() override;
+		
+	protected:
+		
+		//~ Begin FTraceControls Interface
+		virtual ConcertInsightsCore::FStartTraceArgs GetDefaultSynchronizedTraceArgs() const override;
+		virtual bool CanSendRequestsToEndpoint(const FGuid& EndpointId, const IConcertSession& Session) const override;
+		virtual ConcertInsightsCore::FInitArgs GetInitEventArgs() const override;
+		//~ End FTraceControls Interface
+
+	private:
+		
+		FClientTraceControls();
+		
+		void OnSessionStart(TSharedRef<IConcertClientSession> Session);
+		void OnSessionStopped(TSharedRef<IConcertClientSession> Session);
+	};
+}
diff --git a/Engine/Plugins/Developer/Concert/ConcertInsights/ConcertInsightsClient/Source/ConcertInsightsClient/Private/ConcertInsightsClient.cpp b/Engine/Plugins/Developer/Concert/ConcertInsights/ConcertInsightsClient/Source/ConcertInsightsClient/Private/ConcertInsightsClient.cpp
new file mode 100644
--- /dev/null
+++ b/Engine/Plugins/Developer/Concert/ConcertInsights/ConcertInsightsClient/Source/ConcertInsightsClient/Private/ConcertInsightsClient.cpp
@@ -0,0 +1,35 @@
+﻿// Copyright Epic Games, Inc. All Rights Reserved.
+
+#include "ConcertInsightsClient.h"
+
+#include "StatusBar/SMultiUserStatusBar.h"
+#include "StatusBar/StatusBarExtension.h"
+
+#include "Misc/CoreDelegates.h"
+
+#define LOCTEXT_NAMESPACE "FConcertInsightsEditorModule"
+
+namespace UE::ConcertInsightsClient
+{
+	void FConcertInsightsClientModule::StartupModule()
+	{
+		FCoreDelegates::OnPostEngineInit.AddRaw(this, &FConcertInsightsClientModule::PostEngineInit);
+	}
+
+	void FConcertInsightsClientModule::ShutdownModule()
+	{}
+
+	void FConcertInsightsClientModule::PostEngineInit()
+	{
+		// Multi User does not add any status bar by itself. For now, we'll add one but in the future we may want this to be in the MU module.
+		ExtendEditorStatusBarWithMultiUserWidget();
+		
+		TraceControls = ConcertInsightsCore::FTraceControls::Make<FClientTraceControls>();
+		ExtendMultiUserStatusBarWithInsights(*TraceControls);
+	}
+}
+
+
+#undef LOCTEXT_NAMESPACE
+    
+IMPLEMENT_MODULE(UE::ConcertInsightsClient::FConcertInsightsClientModule, ConcertInsightsClient)
\ No newline at end of file
diff --git a/Engine/Plugins/Developer/Concert/ConcertInsights/ConcertInsightsClient/Source/ConcertInsightsClient/Private/ConcertInsightsClient.h b/Engine/Plugins/Developer/Concert/ConcertInsights/ConcertInsightsClient/Source/ConcertInsightsClient/Private/ConcertInsightsClient.h
new file mode 100644
--- /dev/null
+++ b/Engine/Plugins/Developer/Concert/ConcertInsights/ConcertInsightsClient/Source/ConcertInsightsClient/Private/ConcertInsightsClient.h
@@ -0,0 +1,27 @@
+﻿// Copyright Epic Games, Inc. All Rights Reserved.
+
+#pragma once
+
+#include "ClientTraceControls.h"
+#include "IConcertInsightsClientModule.h"
+#include "Templates/UniquePtr.h"
+
+namespace UE::ConcertInsightsClient
+{
+	class FConcertInsightsClientModule : public IConcertInsightsClientModule
+	{
+	public:
+		
+		//~ Begin IModuleInterface Interface
+		virtual void StartupModule() override;
+		virtual void ShutdownModule() override;
+		//~ End IModuleInterface Interface
+
+	private:
+
+		/** The local state of synchronized tracing. */
+		TUniquePtr<FClientTraceControls> TraceControls;
+
+		void PostEngineInit();
+	};
+}
diff --git a/Engine/Plugins/Developer/Concert/ConcertInsights/ConcertInsightsClient/Source/ConcertInsightsClient/Private/ConcertInsightsClientSettings.h b/Engine/Plugins/Developer/Concert/ConcertInsights/ConcertInsightsClient/Source/ConcertInsightsClient/Private/ConcertInsightsClientSettings.h
new file mode 100644
--- /dev/null
+++ b/Engine/Plugins/Developer/Concert/ConcertInsights/ConcertInsightsClient/Source/ConcertInsightsClient/Private/ConcertInsightsClientSettings.h
@@ -0,0 +1,19 @@
+// Copyright Epic Games, Inc. All Rights Reserved.
+
+#pragma once
+
+#include "UObject/Object.h"
+#include "ConcertInsightsClientSettings.generated.h"
+
+UCLASS(config = EditorPerProjectUserSettings)
+class CONCERTINSIGHTSCLIENT_API UConcertInsightsClientSettings : public UObject
+{
+	GENERATED_BODY()
+public:
+
+	static UConcertInsightsClientSettings* Get() { return GetMutableDefault<UConcertInsightsClientSettings>(); }
+
+	/** This is the IP address that is sent to all endpoints to trace to when starting a synchronized trace. */
+	UPROPERTY(config)
+	FString SynchronizedTraceDestinationIP = TEXT("localhost");
+};
diff --git a/Engine/Plugins/Developer/Concert/ConcertInsights/ConcertInsightsClient/Source/ConcertInsightsClient/Private/StatusBar/SEditTraceDestinationWidget.cpp b/Engine/Plugins/Developer/Concert/ConcertInsights/ConcertInsightsClient/Source/ConcertInsightsClient/Private/StatusBar/SEditTraceDestinationWidget.cpp
new file mode 100644
--- /dev/null
+++ b/Engine/Plugins/Developer/Concert/ConcertInsights/ConcertInsightsClient/Source/ConcertInsightsClient/Private/StatusBar/SEditTraceDestinationWidget.cpp
@@ -0,0 +1,53 @@
+// Copyright Epic Games, Inc. All Rights Reserved.
+
+#include "SEditTraceDestinationWidget.h"
+
+#include "ConcertInsightsClientSettings.h"
+#include "Internationalization/Regex.h"
+#include "Widgets/Input/SEditableTextBox.h"
+
+#define LOCTEXT_NAMESPACE "SEditTraceDestinationWidget"
+
+namespace UE::ConcertInsightsClient
+{
+	void SEditTraceDestinationWidget::Construct(const FArguments& InArgs)
+	{
+		ChildSlot
+		[
+			SNew(SEditableTextBox)
+			.MinDesiredWidth(100.f)
+			.Text_Lambda([](){ return FText::FromString(UConcertInsightsClientSettings::Get()->SynchronizedTraceDestinationIP); })
+			.OnTextCommitted_Lambda([](const FText& Text, ETextCommit::Type)
+			{
+				UConcertInsightsClientSettings* Settings = UConcertInsightsClientSettings::Get();
+				const FString TextString = Text.ToString();
+				if (Settings->SynchronizedTraceDestinationIP != TextString)
+				{
+					Settings->SynchronizedTraceDestinationIP = Text.ToString();
+					Settings->SaveConfig();
+				}
+			})
+			.OnVerifyTextChanged_Lambda([](const FText& TextToVerify, FText& Error)
+			{
+				const FString TextString = TextToVerify.ToString();
+				if (TextString.Equals(TEXT("localhost"), ESearchCase::IgnoreCase))
+				{
+					return true;
+				}
+				
+				// Regex for matching ipv4 address stolen from https://www.oreilly.com/library/view/regular-expressions-cookbook/9780596802837/ch07s16.html
+				const FRegexPattern Pattern(TEXT("^(?:[0-9]{1,3}\\.){3}[0-9]{1,3}$"));
+				FRegexMatcher Matcher(Pattern, TextString);
+				
+				const bool bIsIp = Matcher.FindNext();
+				if (!bIsIp)
+				{
+					Error = LOCTEXT("InvalidIp", "Invalid IPv4 Address");
+				}
+				return bIsIp;
+			})
+		];
+	}
+}
+
+#undef LOCTEXT_NAMESPACE
\ No newline at end of file
diff --git a/Engine/Plugins/Developer/Concert/ConcertInsights/ConcertInsightsClient/Source/ConcertInsightsClient/Private/StatusBar/SEditTraceDestinationWidget.h b/Engine/Plugins/Developer/Concert/ConcertInsights/ConcertInsightsClient/Source/ConcertInsightsClient/Private/StatusBar/SEditTraceDestinationWidget.h
new file mode 100644
--- /dev/null
+++ b/Engine/Plugins/Developer/Concert/ConcertInsights/ConcertInsightsClient/Source/ConcertInsightsClient/Private/StatusBar/SEditTraceDestinationWidget.h
@@ -0,0 +1,21 @@
+// Copyright Epic Games, Inc. All Rights Reserved.
+
+#pragma once
+
+#include "Widgets/DeclarativeSyntaxSupport.h"
+#include "Widgets/SCompoundWidget.h"
+
+namespace UE::ConcertInsightsClient
+{
+	class SEditTraceDestinationWidget : public SCompoundWidget
+	{
+	public:
+
+		SLATE_BEGIN_ARGS(SEditTraceDestinationWidget){}
+		SLATE_END_ARGS()
+
+		void Construct(const FArguments& InArgs);
+	};
+}
+
+
diff --git a/Engine/Plugins/Developer/Concert/ConcertInsights/ConcertInsightsClient/Source/ConcertInsightsClient/Private/StatusBar/SMultiUserStatusBar.cpp b/Engine/Plugins/Developer/Concert/ConcertInsights/ConcertInsightsClient/Source/ConcertInsightsClient/Private/StatusBar/SMultiUserStatusBar.cpp
new file mode 100644
--- /dev/null
+++ b/Engine/Plugins/Developer/Concert/ConcertInsights/ConcertInsightsClient/Source/ConcertInsightsClient/Private/StatusBar/SMultiUserStatusBar.cpp
@@ -0,0 +1,173 @@
+// Copyright Epic Games, Inc. All Rights Reserved.
+
+#include "SMultiUserStatusBar.h"
+
+#include "ConcertFrontendStyle.h"
+#include "IConcertSyncClient.h"
+#include "IConcertSyncClientModule.h"
+
+#include "Framework/Docking/TabManager.h"
+#include "Styling/AppStyle.h"
+#include "Styling/StyleColors.h"
+#include "ToolMenus.h"
+#include "ToolMenu.h"
+#include "ToolMenuSection.h"
+#include "Widgets/Images/SImage.h"
+#include "Widgets/Input/SComboButton.h"
+#include "Widgets/Text/STextBlock.h"
+#include "Widgets/SBoxPanel.h"
+#include "Widgets/Layout/SScaleBox.h"
+
+#define LOCTEXT_NAMESPACE "SMultiUserStatusBar"
+
+namespace UE::ConcertInsightsClient
+{
+	void ExtendEditorStatusBarWithMultiUserWidget()
+	{
+		SMultiUserStatusBar::RegisterMultiUserToolMenu();
+		
+		UToolMenu* Menu = UToolMenus::Get()->ExtendMenu(TEXT("LevelEditor.StatusBar.ToolBar"));
+		FToolMenuSection& MultiUserSection = Menu->AddSection(TEXT("MultiUser"), FText::GetEmpty(), FToolMenuInsert(NAME_None, EToolMenuInsertType::First));
+		MultiUserSection.AddEntry(
+			FToolMenuEntry::InitWidget(TEXT("MultiUserStatusBar"), SNew(SMultiUserStatusBar), FText::GetEmpty(), true, false)
+		);
+	}
+
+	void SMultiUserStatusBar::RegisterMultiUserToolMenu()
+	{
+		UToolMenus* ToolMenus = UToolMenus::Get();
+		if (!ToolMenus->IsMenuRegistered(TEXT("MultiUser.StatusBarMenu")))
+		{
+			ToolMenus->RegisterMenu(TEXT("MultiUser.StatusBarMenu"));
+			// Could add additional utilities here, like opening the MU tab
+		}
+	}
+
+	void SMultiUserStatusBar::Construct(const FArguments& InArgs)
+	{
+		ChildSlot
+		[
+			SNew(SHorizontalBox)
+			.ToolTipText_Static(&SMultiUserStatusBar::GetConnectionTooltip)
+			
+			+SHorizontalBox::Slot()
+			.AutoWidth()
+			[
+				MakeComboButton()
+			]
+			
+			+SHorizontalBox::Slot()
+			.AutoWidth()
+			[
+				MakeSessionConnectionIndicator()
+			]
+		];
+	}
+
+	TSharedRef<SWidget> SMultiUserStatusBar::MakeComboButton()
+	{
+		return SNew(SComboButton)
+			.ContentPadding(FMargin(6.0f, 0.0f))
+			.MenuPlacement(MenuPlacement_AboveAnchor)
+			.ComboButtonStyle(&FAppStyle::Get().GetWidgetStyle<FComboButtonStyle>("SimpleComboButton"))
+			.OnGetMenuContent(this, &SMultiUserStatusBar::MakeTraceMenu)
+			.HasDownArrow(true)
+			.ButtonContent()
+			[
+				SNew(SHorizontalBox)
+
+				+SHorizontalBox::Slot()
+				.AutoWidth()
+				.Padding(0, 0, 3, 0)
+				[
+					SNew(SImage)
+					.ColorAndOpacity(FSlateColor::UseForeground())
+					.Image(FConcertFrontendStyle::Get()->GetBrush(TEXT("Concert.MultiUser")))
+				]
+
+				+SHorizontalBox::Slot()
+				.AutoWidth()
+				.VAlign(VAlign_Center)
+				// Last horizontal box should pad to the right so our status bar entry looks consistent with all the other status bar entries
+				.Padding(0, 0, 2, 0)
+				[
+					SNew(STextBlock)
+					.Text(LOCTEXT("MultiUser", "Multi User"))
+					.TextStyle(&FAppStyle::Get().GetWidgetStyle<FTextBlockStyle>("DialogButtonText"))
+				]
+			];
+	}
+
+	TSharedRef<SWidget> SMultiUserStatusBar::MakeSessionConnectionIndicator() const
+	{
+		return SNew(SHorizontalBox)
+		
+			+SHorizontalBox::Slot()
+			.AutoWidth()
+			.Padding(2.f)
+			[
+				SNew(SScaleBox)
+				[
+					SNew(SImage)
+					.Image(FAppStyle::GetBrush("Icons.FilledCircle"))
+					.ColorAndOpacity_Lambda([this]()
+					{
+						return IsConnectedToSession()
+							? FStyleColors::AccentGreen
+							: FStyleColors::AccentGray;
+					})
+				]
+			]
+		
+			+SHorizontalBox::Slot()
+			.AutoWidth()
+			.VAlign(VAlign_Center)
+			[
+				SNew(STextBlock)
+				.Text_Static(&SMultiUserStatusBar::GetCurrentSessionLabel)
+			]
+		;
+	}
+
+	TSharedRef<SWidget> SMultiUserStatusBar::MakeTraceMenu() const
+	{
+		UToolMenus* ToolMenus = UToolMenus::Get();
+		if (!ToolMenus->IsMenuRegistered(TEXT("MultiUser.StatusBarMenu")))
+		{
+			ToolMenus->RegisterMenu(TEXT("MultiUser.StatusBarMenu"));
+			// Could add additional utilities here, like opening the MU tab
+		}
+
+		FToolMenuContext Context;
+		UToolMenu* Menu = UToolMenus::Get()->GenerateMenu(TEXT("MultiUser.StatusBarMenu"), Context);
+		return ToolMenus->GenerateWidget(Menu);
+	}
+
+	bool SMultiUserStatusBar::IsConnectedToSession()
+	{
+		const TSharedPtr<IConcertSyncClient> Client = IConcertSyncClientModule::Get().GetClient(TEXT("MultiUser"));
+		return Client && Client->GetConcertClient()->GetCurrentSession().IsValid();
+	}
+
+	FText SMultiUserStatusBar::GetConnectionTooltip()
+	{
+		const TSharedPtr<IConcertSyncClient> Client = IConcertSyncClientModule::Get().GetClient(TEXT("MultiUser"));
+		const TSharedPtr<IConcertClientSession> Session = Client->GetConcertClient()->GetCurrentSession();
+		const bool bIsConnected = Session.IsValid();
+		return bIsConnected
+			? FText::Format(LOCTEXT("ConnectionIndicator.OnlineFmt", "Connected to session {0}"), FText::FromString(Session->GetSessionInfo().SessionName))
+			: LOCTEXT("ConnectionIndicator.Offline", "Disconnected");
+	}
+
+	FText SMultiUserStatusBar::GetCurrentSessionLabel()
+	{
+		const TSharedPtr<IConcertSyncClient> Client = IConcertSyncClientModule::Get().GetClient(TEXT("MultiUser"));
+		const TSharedPtr<IConcertClientSession> Session = Client->GetConcertClient()->GetCurrentSession();
+		const bool bIsConnected = Session.IsValid();
+		return bIsConnected
+			? FText::FromString(Session->GetSessionInfo().SessionName)
+			: LOCTEXT("CurrentSession.None", "No session");
+	}
+}
+
+#undef LOCTEXT_NAMESPACE
\ No newline at end of file
diff --git a/Engine/Plugins/Developer/Concert/ConcertInsights/ConcertInsightsClient/Source/ConcertInsightsClient/Private/StatusBar/SMultiUserStatusBar.h b/Engine/Plugins/Developer/Concert/ConcertInsights/ConcertInsightsClient/Source/ConcertInsightsClient/Private/StatusBar/SMultiUserStatusBar.h
new file mode 100644
--- /dev/null
+++ b/Engine/Plugins/Developer/Concert/ConcertInsights/ConcertInsightsClient/Source/ConcertInsightsClient/Private/StatusBar/SMultiUserStatusBar.h
@@ -0,0 +1,36 @@
+// Copyright Epic Games, Inc. All Rights Reserved.
+
+#pragma once
+
+#include "Widgets/DeclarativeSyntaxSupport.h"
+#include "Widgets/SCompoundWidget.h"
+
+namespace UE::ConcertInsightsClient
+{
+	/** Extends the editor status bar with the Multi-User widget. */
+	void ExtendEditorStatusBarWithMultiUserWidget();
+	
+	class SMultiUserStatusBar : public SCompoundWidget
+	{
+	public:
+
+		/** Register the UToolMenu for the combo button. */
+		static void RegisterMultiUserToolMenu();
+
+		SLATE_BEGIN_ARGS(SMultiUserStatusBar) {}
+		SLATE_END_ARGS()
+
+		void Construct(const FArguments& InArgs);
+
+	private:
+		TSharedRef<SWidget> MakeComboButton();
+		TSharedRef<SWidget> MakeSessionConnectionIndicator() const;
+		
+		TSharedRef<SWidget> MakeTraceMenu() const;
+
+		static bool IsConnectedToSession();
+		static FText GetConnectionTooltip();
+		static FText GetCurrentSessionLabel();
+	};
+}
+
diff --git a/Engine/Plugins/Developer/Concert/ConcertInsights/ConcertInsightsClient/Source/ConcertInsightsClient/Private/StatusBar/StatusBarExtension.cpp b/Engine/Plugins/Developer/Concert/ConcertInsights/ConcertInsightsClient/Source/ConcertInsightsClient/Private/StatusBar/StatusBarExtension.cpp
new file mode 100644
--- /dev/null
+++ b/Engine/Plugins/Developer/Concert/ConcertInsights/ConcertInsightsClient/Source/ConcertInsightsClient/Private/StatusBar/StatusBarExtension.cpp
@@ -0,0 +1,114 @@
+// Copyright Epic Games, Inc. All Rights Reserved.
+
+#include "StatusBarExtension.h"
+
+#include "ConcertInsightsClient.h"
+#include "IConcertSyncClient.h"
+#include "IConcertSyncClientModule.h"
+#include "SEditTraceDestinationWidget.h"
+
+#include "ToolMenus.h"
+#include "ToolMenu.h"
+#include "ToolMenuSection.h"
+#include "Framework/Notifications/NotificationManager.h"
+#include "Widgets/Notifications/SNotificationList.h"
+
+#define LOCTEXT_NAMESPACE "MultiUserStatusBarExtension"
+
+namespace UE::ConcertInsightsClient
+{
+	namespace Private
+	{
+		static bool IsInSession()
+		{
+			const TSharedPtr<IConcertSyncClient> Client = IConcertSyncClientModule::Get().GetClient(TEXT("MultiUser"));
+			const TSharedPtr<IConcertClientSession> Session = Client->GetConcertClient()->GetCurrentSession();
+			return Session.IsValid();
+		}
+
+		static FToolMenuEntry MakeSynchronizedTraceEntry(FClientTraceControls& Controls)
+		{
+			FToolMenuEntry ConcertTraceEntry = FToolMenuEntry::InitMenuEntry(
+				TEXT("StartConcertTrace"),
+				TAttribute<FText>::CreateLambda([&Controls]()
+				{
+					return Controls.IsTracing()
+						? LOCTEXT("Menu.Tracing.ToggleSynchronizedTrace.Label", "Stop synchronized trace")
+						: LOCTEXT("Menu.Tracing.ToggleSynchronizedTrace.Label", "Start synchronized trace");
+				}),
+				TAttribute<FText>::CreateLambda([&Controls]()
+				{
+					if (!Private::IsInSession())
+					{
+						return LOCTEXT("Menu.Tracing.ToggleSynchronizedTrace.Tooltip.NotInSession", "Not in any Multi-User session. Join a session first.");
+					}
+					
+					return Controls.IsTracing()
+						? LOCTEXT("Menu.Tracing.ToggleSynchronizedTrace.Tooltip.StopTrace", "Stops the synchronized trace across all participants")
+						: LOCTEXT("Menu.Tracing.ToggleSynchronizedTrace.Tooltip.StartTrace", "Starts a synchronized trace across endpoints in the current Multi-User session.");
+				}),
+				FSlateIcon(),
+				FUIAction(
+					FExecuteAction::CreateLambda([&Controls]()
+					{
+						const TSharedPtr<IConcertSyncClient> Client = IConcertSyncClientModule::Get().GetClient(TEXT("MultiUser"));
+						const TSharedPtr<IConcertClientSession> Session = Client->GetConcertClient()->GetCurrentSession();
+						if (!Session)
+						{
+							return;
+						}
+						
+						if (Controls.IsTracing())
+						{
+							Controls.StopSynchronizedTrace();
+						}
+						else
+						{
+							FText ErrorReason;
+							const bool bSuccess = Controls.StartSynchronizedTrace(Session.ToSharedRef(), &ErrorReason);
+							if (!bSuccess)
+							{
+								FNotificationInfo Notification(LOCTEXT("NotificationTitle", "Synchronized Trace Failed"));
+								Notification.SubText = ErrorReason;
+								Notification.bFireAndForget = true;
+								Notification.ExpireDuration = 4.f;
+								FSlateNotificationManager::Get().AddNotification(Notification)
+									->SetCompletionState(SNotificationItem::CS_Fail);
+							}
+						}
+					}),
+					FCanExecuteAction::CreateLambda([]() { return Private::IsInSession();})
+					)
+			);
+			ConcertTraceEntry.bShouldCloseWindowAfterMenuSelection = false;
+
+			return ConcertTraceEntry;
+		}
+		
+		static FToolMenuEntry MakeTraceDestinationIpEntry()
+		{
+			return FToolMenuEntry::InitWidget(
+				TEXT("EditIp"),
+				SNew(SEditTraceDestinationWidget),
+				LOCTEXT("EditIp.Label", "Destination IP"),
+				true, true, false,
+				LOCTEXT("EditIp.ToolTip", "Enter the IP of the trace store to send all trace data to. Tracing will fail to start if the IP is invalid.")
+				);
+		}
+	}
+	
+	void ExtendMultiUserStatusBarWithInsights(FClientTraceControls& Controls)
+	{
+		UToolMenu* ToolMenu = UToolMenus::Get()->ExtendMenu(TEXT("MultiUser.StatusBarMenu"));
+		if (!ToolMenu)
+		{
+			return;
+		}
+
+		FToolMenuSection& ControlsSection = ToolMenu->AddSection(TEXT("Tracing"), LOCTEXT("Menu.Tracing", "Tracing"));
+		ControlsSection.AddEntry(Private::MakeSynchronizedTraceEntry(Controls));
+		ControlsSection.AddEntry(Private::MakeTraceDestinationIpEntry());
+	}
+}
+
+#undef LOCTEXT_NAMESPACE
\ No newline at end of file
diff --git a/Engine/Plugins/Developer/Concert/ConcertInsights/ConcertInsightsClient/Source/ConcertInsightsClient/Private/StatusBar/StatusBarExtension.h b/Engine/Plugins/Developer/Concert/ConcertInsights/ConcertInsightsClient/Source/ConcertInsightsClient/Private/StatusBar/StatusBarExtension.h
new file mode 100644
--- /dev/null
+++ b/Engine/Plugins/Developer/Concert/ConcertInsights/ConcertInsightsClient/Source/ConcertInsightsClient/Private/StatusBar/StatusBarExtension.h
@@ -0,0 +1,14 @@
+// Copyright Epic Games, Inc. All Rights Reserved.
+
+#pragma once
+
+#include "HAL/Platform.h"
+
+namespace UE::ConcertInsightsClient
+{
+	class FClientTraceControls;
+	
+	/** Adds a sub-menu to the Multi User status bar. */
+	void ExtendMultiUserStatusBarWithInsights(FClientTraceControls& Controls UE_LIFETIMEBOUND);
+}
+
diff --git a/Engine/Plugins/Developer/Concert/ConcertInsights/ConcertInsightsClient/Source/ConcertInsightsClient/Public/IConcertInsightsClientModule.h b/Engine/Plugins/Developer/Concert/ConcertInsights/ConcertInsightsClient/Source/ConcertInsightsClient/Public/IConcertInsightsClientModule.h
new file mode 100644
--- /dev/null
+++ b/Engine/Plugins/Developer/Concert/ConcertInsights/ConcertInsightsClient/Source/ConcertInsightsClient/Public/IConcertInsightsClientModule.h
@@ -0,0 +1,42 @@
+// Copyright Epic Games, Inc. All Rights Reserved.
+
+#pragma once
+
+#include "Modules/ModuleManager.h"
+
+namespace UE::ConcertInsightsClient
+{
+	/**
+	 * On the editor, this module
+	 * - adds a context menu option to the Unreal Insights editor menu (bottom-right) for starting synchronized session tracing across multiple machines
+	 * - listens for requests to synchronized session tracing while the local editor is in a Concert session
+	 */
+	class CONCERTINSIGHTSCLIENT_API IConcertInsightsClientModule : public IModuleInterface
+	{
+	public:
+
+		/**
+		 * Singleton-like access to this module's interface.  This is just for convenience!
+		 * Beware of calling this during the shutdown phase, though.  Your module might have been unloaded already.
+		 *
+		 * @return Returns singleton instance, loading the module on demand if needed
+		 */
+		static inline IConcertInsightsClientModule& Get()
+		{
+			static const FName ModuleName = "ConcertInsightsClient";
+			return FModuleManager::LoadModuleChecked<IConcertInsightsClientModule>(ModuleName);
+		}
+
+		/**
+		 * Checks to see if this module is loaded and ready.  It is only valid to call Get() during shutdown if IsAvailable() returns true.
+		 *
+		 * @return True if the module is loaded and ready to use
+		 */
+		static inline bool IsAvailable()
+		{
+			static const FName ModuleName = "ConcertInsightsClient";
+			return FModuleManager::Get().IsModuleLoaded(ModuleName);
+		}
+	};
+}
+
diff --git a/Engine/Plugins/Developer/Concert/ConcertInsights/ConcertInsightsCore/ConcertInsightsCore.uplugin b/Engine/Plugins/Developer/Concert/ConcertInsights/ConcertInsightsCore/ConcertInsightsCore.uplugin
new file mode 100644
--- /dev/null
+++ b/Engine/Plugins/Developer/Concert/ConcertInsights/ConcertInsightsCore/ConcertInsightsCore.uplugin
@@ -0,0 +1,39 @@
+﻿{
+	"FileVersion": 3,
+	"Version": 1,
+	"VersionName": "1.0",
+	"FriendlyName": "ConcertInsightsCore",
+	"Description": "Shared logic for starting synchronized tracing",
+	"Category": "Networking",
+	"CreatedBy": "Epic Games, Inc.",
+	"CreatedByURL": "",
+	"DocsURL": "",
+	"MarketplaceURL": "",
+	"CanContainContent": false,
+	"Hidden": true,
+	"IsBetaVersion": false,
+	"IsExperimentalVersion": true,
+	"Installed": false,
+	"SupportedPrograms": [
+		"UnrealInsights",
+		"UnrealMultiUserServer",
+		"UnrealMultiUserSlateServer"
+	],
+	"Modules": [
+		{
+			"Name": "ConcertInsightsCore",
+			"Type": "EditorAndProgram",
+			"LoadingPhase": "Default",
+			"ProgramAllowList": [
+				"UnrealMultiUserServer",
+				"UnrealMultiUserSlateServer"
+			]
+		}
+	],
+	"Plugins": [
+		{
+			"Name": "ConcertMain",
+			"Enabled": true
+		}
+	]
+}
\ No newline at end of file
diff --git a/Engine/Plugins/Developer/Concert/ConcertInsights/ConcertInsightsCore/Source/ConcertInsightsCore/ConcertInsightsCore.Build.cs b/Engine/Plugins/Developer/Concert/ConcertInsights/ConcertInsightsCore/Source/ConcertInsightsCore/ConcertInsightsCore.Build.cs
new file mode 100644
--- /dev/null
+++ b/Engine/Plugins/Developer/Concert/ConcertInsights/ConcertInsightsCore/Source/ConcertInsightsCore/ConcertInsightsCore.Build.cs
@@ -0,0 +1,32 @@
+﻿// Copyright Epic Games, Inc. All Rights Reserved.
+
+using UnrealBuildTool;
+
+public class ConcertInsightsCore : ModuleRules
+{
+    public ConcertInsightsCore(ReadOnlyTargetRules Target) : base(Target)
+    {
+        PCHUsage = ModuleRules.PCHUsageMode.UseExplicitOrSharedPCHs;
+
+        PublicDependencyModuleNames.AddRange(
+            new string[]
+            {
+                "Core",
+            }
+        );
+
+        PrivateDependencyModuleNames.AddRange(
+            new string[]
+            {
+	            "Concert",
+	            "ConcertTransport",
+                "CoreUObject",
+                "Slate",
+                "SlateCore",
+                "TraceLog",
+            }
+        );
+
+        ShortName = "CrtInC";
+    }
+}
\ No newline at end of file
diff --git a/Engine/Plugins/Developer/Concert/ConcertInsights/ConcertInsightsCore/Source/ConcertInsightsCore/Private/ConcertInsightsSynchronizedTrace.cpp b/Engine/Plugins/Developer/Concert/ConcertInsights/ConcertInsightsCore/Source/ConcertInsightsCore/Private/ConcertInsightsSynchronizedTrace.cpp
new file mode 100644
--- /dev/null
+++ b/Engine/Plugins/Developer/Concert/ConcertInsights/ConcertInsightsCore/Source/ConcertInsightsCore/Private/ConcertInsightsSynchronizedTrace.cpp
@@ -0,0 +1,23 @@
+﻿// Copyright Epic Games, Inc. All Rights Reserved.
+
+#include "ConcertInsightsSynchronizedTrace.h"
+
+#define LOCTEXT_NAMESPACE "FConcertInsightsSynchronizedTraceModule"
+
+namespace UE::ConcertInsightsSyncTrace
+{
+	void FConcertInsightsSyncTraceModule::StartupModule()
+	{
+    
+	}
+
+	void FConcertInsightsSyncTraceModule::ShutdownModule()
+	{
+    
+	}
+}
+
+
+#undef LOCTEXT_NAMESPACE
+    
+IMPLEMENT_MODULE(UE::ConcertInsightsSyncTrace::FConcertInsightsSyncTraceModule, ConcertInsightsCore)
\ No newline at end of file
diff --git a/Engine/Plugins/Developer/Concert/ConcertInsights/ConcertInsightsCore/Source/ConcertInsightsCore/Private/ConcertInsightsSynchronizedTrace.h b/Engine/Plugins/Developer/Concert/ConcertInsights/ConcertInsightsCore/Source/ConcertInsightsCore/Private/ConcertInsightsSynchronizedTrace.h
new file mode 100644
--- /dev/null
+++ b/Engine/Plugins/Developer/Concert/ConcertInsights/ConcertInsightsCore/Source/ConcertInsightsCore/Private/ConcertInsightsSynchronizedTrace.h
@@ -0,0 +1,19 @@
+﻿// Copyright Epic Games, Inc. All Rights Reserved.
+
+#pragma once
+
+#include "IConcertInsightsSyncTraceModule.h"
+
+namespace UE::ConcertInsightsSyncTrace
+{
+	class FConcertInsightsSyncTraceModule : public IConcertInsightsSyncTraceModule
+	{
+	public:
+
+		//~ Begin IModuleInterface Interface
+		virtual void StartupModule() override;
+		virtual void ShutdownModule() override;
+		//~ End IModuleInterface Interface
+	};
+}
+
diff --git a/Engine/Plugins/Developer/Concert/ConcertInsights/ConcertInsightsCore/Source/ConcertInsightsCore/Private/TraceControls.cpp b/Engine/Plugins/Developer/Concert/ConcertInsights/ConcertInsightsCore/Source/ConcertInsightsCore/Private/TraceControls.cpp
new file mode 100644
--- /dev/null
+++ b/Engine/Plugins/Developer/Concert/ConcertInsights/ConcertInsightsCore/Source/ConcertInsightsCore/Private/TraceControls.cpp
@@ -0,0 +1,222 @@
+// Copyright Epic Games, Inc. All Rights Reserved.
+
+#include "TraceControls.h"
+
+#include "ConcertLogGlobal.h"
+#include "Concert/Public/IConcertSession.h"
+#include "Trace/ConcertProtocolTrace.h"
+
+#include "ProfilingDebugging/TraceAuxiliary.h"
+
+#define LOCTEXT_NAMESPACE "FTraceControls"
+
+namespace UE::ConcertInsightsCore
+{
+	FTraceControls::FTraceControls()
+	{
+		FTraceAuxiliary::OnTraceStarted.AddRaw(this, &FTraceControls::OnTraceStarted);
+		FTraceAuxiliary::OnTraceStopped.AddRaw(this, &FTraceControls::OnTraceStopped);
+	}
+
+	FTraceControls::~FTraceControls()
+	{
+		FTraceAuxiliary::OnTraceStopped.RemoveAll(this);
+	}
+
+	bool FTraceControls::IsTracing() const
+	{
+		return FTraceAuxiliary::IsConnected();
+	}
+
+#define SET_REASON(FailReasonVar, ReasonLocText) if (FailReason) { *FailReasonVar = ReasonLocText; }
+	
+	bool FTraceControls::StartSynchronizedTrace(TSharedRef<IConcertSession> Session, const FStartTraceArgs& Args, FText* FailReason)
+	{
+		if (IsTracing())
+		{
+			SET_REASON(FailReason, LOCTEXT("Reason.AlreadyRunning", "A trace is already in progress"));
+			return false;
+		}
+
+		if (!LocallyStartSynchronizedTrace({ Session }, Args))
+		{
+			SET_REASON(FailReason, LOCTEXT("Reason.AlreadyRunning", "Failed to start local trace. See log."));
+			return false;
+		}
+		
+		FConcertTrace_StartSyncTrace_Request Request;
+		Request.TraceArgs = static_cast<FConcertTrace_StartTraceArgs>(Args);
+		auto SendRequest = [this, &Session, &Request](const FGuid& Endpoint)
+		{
+			if (CanSendRequestsToEndpoint(Endpoint, *Session))
+			{
+				UE_LOG(LogConcert, Verbose, TEXT("Sending synchronized trace request to client %s"), *Endpoint.ToString());
+				
+				Session->SendCustomRequest<FConcertTrace_StartSyncTrace_Request, FConcertTrace_StartSyncTrace_Response>(Request, Endpoint)
+					.Next([Endpoint](FConcertTrace_StartSyncTrace_Response&& Response)
+					{
+						UE_CLOG(Response.ErrorCode == EConcertTraceErrorCode::Timeout, LogConcert, Warning, TEXT("Client %s timed out answering synchronized trace request"), *Endpoint.ToString());
+						UE_CLOG(Response.ErrorCode == EConcertTraceErrorCode::Joined, LogConcert, Log, TEXT("Client %s accepted synchronized trace request"), *Endpoint.ToString());
+						UE_CLOG(Response.ErrorCode == EConcertTraceErrorCode::Rejected, LogConcert, Error, TEXT("Client %s rejected synchronized trace request"), *Endpoint.ToString());
+					});
+			}
+		};
+		
+		SendRequest(Session->GetSessionInfo().ServerEndpointId);
+		for (const FGuid& ClientEndpoint : Session->GetSessionClientEndpointIds())
+		{
+			SendRequest(ClientEndpoint);
+		}
+
+		return true;
+	}
+
+	void FTraceControls::StopSynchronizedTrace()
+	{
+		if (!InProgressSynchronizedTrace)
+		{
+			return;
+		}
+		UE_LOG(LogConcert, Log, TEXT("Stopping synchronized trace..."));
+		
+		const TSharedPtr<IConcertSession> Session = InProgressSynchronizedTrace->TraceInitiator.Pin();
+		// Reset before calling FTraceAuxiliary::Stop() so OnTraceStopped does not cause a recursive call
+		InProgressSynchronizedTrace.Reset();
+		
+		if (Session)
+		{
+			FTraceAuxiliary::Stop();
+			
+			auto SendRequest = [this, &Session](const FGuid& Endpoint)
+			{
+				Session->SendCustomEvent(FConcertTrace_StopSyncTrace{}, { Endpoint } , EConcertMessageFlags::ReliableOrdered);
+			};
+
+			const FGuid ServerEndpoint = Session->GetSessionInfo().ServerEndpointId;
+			if (CanSendRequestsToEndpoint(ServerEndpoint, *Session))
+			{
+				SendRequest(ServerEndpoint);
+			}
+			for (const FGuid& ClientEndpointId : Session->GetSessionClientEndpointIds())
+			{
+				SendRequest(ClientEndpointId);
+			}
+		}
+		
+		OnSynchronizedTraceStoppedDelegate.Broadcast();
+	}
+
+	void FTraceControls::RegisterTraceRequestsHandler(const TSharedRef<IConcertSession>& Session)
+	{
+		if (!ensure(!RegisteredSessions.Contains(Session)))
+		{
+			return;
+		}
+		RegisteredSessions.Add(Session);
+		
+		Session->RegisterCustomRequestHandler<FConcertTrace_StartSyncTrace_Request, FConcertTrace_StartSyncTrace_Response>(
+			[this, WeakSession = TWeakPtr<IConcertSession>(Session)](const FConcertSessionContext& Context, const FConcertTrace_StartSyncTrace_Request& Request, FConcertTrace_StartSyncTrace_Response& Response)
+			{
+				const TSharedPtr<IConcertSession> SessionPin = WeakSession.Pin();
+				check(SessionPin);
+				return HandleTraceStartRequest(Context, Request, Response, SessionPin.ToSharedRef());
+			});
+		Session->RegisterCustomEventHandler<FConcertTrace_StopSyncTrace>(this, &FTraceControls::HandleTraceStopRequest);
+	}
+
+	void FTraceControls::OnLeaveSession(IConcertSession& Session)
+	{
+		if (IsInSynchronizedTrace())
+		{
+			StopLocalConcertTrace();
+		}
+	}
+
+	EConcertSessionResponseCode FTraceControls::HandleTraceStartRequest(
+		const FConcertSessionContext& Context,
+		const FConcertTrace_StartSyncTrace_Request& Request,
+		FConcertTrace_StartSyncTrace_Response& Response,
+		const TSharedRef<IConcertSession>& Session
+		)
+	{
+		if (!FTraceAuxiliary::IsConnected())
+		{
+			Response.ErrorCode = EConcertTraceErrorCode::Joined;
+
+			LocallyStartSynchronizedTrace({ Session }, FStartTraceArgs{{ Request.TraceArgs }});
+			OnSynchronizedTraceAccepted(Context, Request, Session);
+		}
+		else
+		{
+			Response.ErrorCode = EConcertTraceErrorCode::Rejected;
+		}
+		
+		return EConcertSessionResponseCode::Success;
+	}
+
+	void FTraceControls::OnTraceStopped(FTraceAuxiliary::EConnectionType, const FString&)
+	{
+		if (InProgressSynchronizedTrace.IsSet())
+		{
+			StopSynchronizedTrace();
+		}
+	}
+
+	bool FTraceControls::StartLocalConcertTrace(const FStartTraceArgs& Args)
+	{
+		Trace::ToggleChannel(TEXT("ConcertChannel"), true);
+		
+		if (!FTraceAuxiliary::IsConnected())
+		{
+			return FTraceAuxiliary::Start(
+				ConcertInsightsSync::ConvertTraceTargetType(Args.ConnectionType),
+				*Args.Target,
+				*Args.Channels,
+				Args.Options,
+				Args.LogCategory
+				);
+		}
+
+		// Trace already running, so successful.
+		SendInitEventIfNeeded();
+		return true;
+	}
+
+	void FTraceControls::StopLocalConcertTrace()
+	{
+		FTraceAuxiliary::Stop();
+	}
+
+	void FTraceControls::SendInitEventIfNeeded()
+	{
+		if (!FTraceAuxiliary::IsConnected())
+		{
+			return;
+		}
+		
+		const TOptional<FInitArgs> InitArgs = GetInitEventArgs();
+		if (!InitArgs)
+		{
+			return;
+		}
+		
+		CONCERT_TRACE_INIT(InitArgs->EndpointId, InitArgs->DisplayString, InitArgs->bIsServer);
+	}
+
+	bool FTraceControls::LocallyStartSynchronizedTrace(FSynchronizedSessionData Data, const FStartTraceArgs& Args)
+	{
+		UE_LOG(LogConcert, Log, TEXT("Starting synchronized trace."));
+		
+		check(!InProgressSynchronizedTrace.IsSet());
+		InProgressSynchronizedTrace = MoveTemp(Data);
+
+		if (StartLocalConcertTrace(Args))
+		{
+			OnSynchronizedTraceStartedDelegate.Broadcast();
+			return true;
+		}
+		return false;
+	}
+}
+
+#undef LOCTEXT_NAMESPACE
\ No newline at end of file
diff --git a/Engine/Plugins/Developer/Concert/ConcertInsights/ConcertInsightsCore/Source/ConcertInsightsCore/Public/IConcertInsightsSyncTraceModule.h b/Engine/Plugins/Developer/Concert/ConcertInsights/ConcertInsightsCore/Source/ConcertInsightsCore/Public/IConcertInsightsSyncTraceModule.h
new file mode 100644
--- /dev/null
+++ b/Engine/Plugins/Developer/Concert/ConcertInsights/ConcertInsightsCore/Source/ConcertInsightsCore/Public/IConcertInsightsSyncTraceModule.h
@@ -0,0 +1,47 @@
+// Copyright Epic Games, Inc. All Rights Reserved.
+
+#pragma once
+
+#include "Modules/ModuleManager.h"
+
+namespace UE::ConcertInsightsSyncTrace
+{
+	/**
+	 * This module provides the base code for starting synchronized tracing across multiple machines.
+	 * 
+	 * Synchronized tracing means that certain machines in the session are requested to start tracing at the same time.
+	 * Each machine will generate a separate .utrace file. When any of these files are analyzed in Unreal Insights, ConcertInsights collects the
+	 * other relevant files and aggregates them in the UI.
+	 * 
+	 * This module houses the shared code needed to synchronize events across multiple machines, for example the Concert events that are sent to the other machines.
+	 * ConcertInsightsServer and ConcertInsightsEditor depend on this module's exposed events.
+	 */
+	class CONCERTINSIGHTSCORE_API IConcertInsightsSyncTraceModule : public IModuleInterface
+	{
+	public:
+
+		/**
+		 * Singleton-like access to this module's interface.  This is just for convenience!
+		 * Beware of calling this during the shutdown phase, though.  Your module might have been unloaded already.
+		 *
+		 * @return Returns singleton instance, loading the module on demand if needed
+		 */
+		static inline IConcertInsightsSyncTraceModule& Get()
+		{
+			static const FName ModuleName = "ConcertInsightsCore";
+			return FModuleManager::LoadModuleChecked<IConcertInsightsSyncTraceModule>(ModuleName);
+		}
+
+		/**
+		 * Checks to see if this module is loaded and ready.  It is only valid to call Get() during shutdown if IsAvailable() returns true.
+		 *
+		 * @return True if the module is loaded and ready to use
+		 */
+		static inline bool IsAvailable()
+		{
+			static const FName ModuleName = "ConcertInsightsCore";
+			return FModuleManager::Get().IsModuleLoaded(ModuleName);
+		}
+	};
+}
+
diff --git a/Engine/Plugins/Developer/Concert/ConcertInsights/ConcertInsightsCore/Source/ConcertInsightsCore/Public/TraceControls.h b/Engine/Plugins/Developer/Concert/ConcertInsights/ConcertInsightsCore/Source/ConcertInsightsCore/Public/TraceControls.h
new file mode 100644
--- /dev/null
+++ b/Engine/Plugins/Developer/Concert/ConcertInsights/ConcertInsightsCore/Source/ConcertInsightsCore/Public/TraceControls.h
@@ -0,0 +1,131 @@
+// Copyright Epic Games, Inc. All Rights Reserved.
+
+#pragma once
+
+#include "TraceMessages.h"
+#include "Concert/Public/ConcertMessages.h"
+#include "Concert/Public/IConcertSessionHandler.h"
+#include "Templates/UnrealTemplate.h"
+
+class IConcertSession;
+
+namespace UE::ConcertInsightsCore
+{
+	struct FStartTraceArgs : FConcertTrace_StartTraceArgs
+	{
+		FTraceAuxiliary::FOptions* Options = nullptr;
+		const FTraceAuxiliary::FLogCategoryAlias& LogCategory = LogCore;
+	};
+
+	struct FInitArgs
+	{
+		TOptional<FGuid> EndpointId;
+		FString DisplayString;
+		bool bIsServer;
+	};
+
+	/**
+	 * Manages synchronized traces.
+	 * 
+	 * This listens for remote requests for starting traces.
+	 * Intended to be subclassed.
+	 */
+	class CONCERTINSIGHTSCORE_API FTraceControls : public FNoncopyable
+	{
+	public:
+
+		template<typename TControls, typename... TArgs>
+		static TUniquePtr<TControls> Make(TArgs&&... Arg)
+		{
+			TUniquePtr<TControls> Result(new TControls(Forward<TArgs>(Arg)...));
+			
+			// Ensure that init event is traced if the engine was started with tracing enabled.
+			// This calls virtual functions so must call after constructor.
+			Result->SendInitEventIfNeeded();
+			
+			return Result;
+		}
+		
+		virtual ~FTraceControls();
+
+		/** Whether a trace is currently occuring */
+		bool IsTracing() const;
+
+		/** Starts a synchronized trace across multiple endpoints. */
+		bool StartSynchronizedTrace(TSharedRef<IConcertSession> Session, const FStartTraceArgs& Args, FText* FailReason = nullptr);
+		bool StartSynchronizedTrace(TSharedRef<IConcertSession> Session, FText* FailReason = nullptr) { return StartSynchronizedTrace(Session, GetDefaultSynchronizedTraceArgs(), FailReason); }
+
+		/** Stops a synchronized trace if one is ongoing. */
+		void StopSynchronizedTrace();
+
+		/** Gets the default trace arguments to use */
+		virtual FStartTraceArgs GetDefaultSynchronizedTraceArgs() const { return {}; }
+
+		bool IsInSynchronizedTrace() const { return InProgressSynchronizedTrace.IsSet(); }
+
+		DECLARE_EVENT(FTraceControls, FOnSynchronizedTraceStarted);
+		/** Broadcasts when a synchronized trace is started for any reason, such as the local machine starting it or a remote request being accepted. */
+		FOnSynchronizedTraceStarted& OnSynchronizedTraceStarted() { return OnSynchronizedTraceStartedDelegate; };
+		
+		DECLARE_EVENT(FTraceControls, FOnSynchronizedTraceStopped);
+		/** Broadcasts when a synchronized trace is stopped for any reason, such as the local machine stopping it or being told so by a remote request. */
+		FOnSynchronizedTraceStopped& OnSynchronizedTraceStopped() { return OnSynchronizedTraceStoppedDelegate; };
+		
+	protected:
+		
+		FTraceControls();
+		
+		/** Registers handlers for trace requests. */
+		void RegisterTraceRequestsHandler(const TSharedRef<IConcertSession>& Session);
+
+		/** Called by subclasses to notify that the local application has left the session. */
+		void OnLeaveSession(IConcertSession& Session);
+
+		/** @return Allows subclass to decide whether a trace start request should be sent to this endpoint. */
+		virtual bool CanSendRequestsToEndpoint(const FGuid& EndpointId, const IConcertSession& Session) const { return true; }
+		/** Called when an incoming synchronized trace was accepted. */
+		virtual void OnSynchronizedTraceAccepted(const FConcertSessionContext& Context, const FConcertTrace_StartSyncTrace_Request& Request, const TSharedRef<IConcertSession>& Session) {}
+		
+		/** Checks whether joining a trace is ok and if so, return the data to put into the ConcertInsights init event. */
+		virtual FInitArgs GetInitEventArgs() const = 0;
+
+	private:
+
+		struct FSynchronizedSessionData
+		{
+			/** The session that the synchronized trace was started on. */
+			TWeakPtr<IConcertSession> TraceInitiator;
+		};
+		
+		/** Sessions these controls are listening for requests on. */
+		TSet<TWeakPtr<IConcertSession>> RegisteredSessions;
+		
+		/** Data for the currently running synchronized trace. Unset if not running. */
+		TOptional<FSynchronizedSessionData> InProgressSynchronizedTrace;
+
+		FOnSynchronizedTraceStarted OnSynchronizedTraceStartedDelegate;
+		FOnSynchronizedTraceStopped OnSynchronizedTraceStoppedDelegate;
+
+		/** Handles networked requests from other endpoints to start tracing. */
+		EConcertSessionResponseCode HandleTraceStartRequest(const FConcertSessionContext& Context, const FConcertTrace_StartSyncTrace_Request& Request, FConcertTrace_StartSyncTrace_Response& Response, const TSharedRef<IConcertSession>& Session);
+		/** Handles networked requests from other endpoints to stop tracing. */
+		void HandleTraceStopRequest(const FConcertSessionContext&, const FConcertTrace_StopSyncTrace&) { StopLocalConcertTrace(); }
+		
+		/** Local callback for when tracing begins */
+		void OnTraceStarted(FTraceAuxiliary::EConnectionType, const FString&) { SendInitEventIfNeeded(); }
+		void OnTraceStopped(FTraceAuxiliary::EConnectionType, const FString&);
+		
+		/** Starts tracing (if not already), enables the Concert trace channel, and ensures everything for ConcertInsights is set up (init event is sent). */
+		bool StartLocalConcertTrace(const FStartTraceArgs& Args);
+		/** Stops tracing locally */
+		void StopLocalConcertTrace();
+
+		/** Sends the init event if we're tracing, the Concert channel is enabled, and the init event has not yet been sent during this session. */
+		void SendInitEventIfNeeded();
+		
+		bool LocallyStartSynchronizedTrace(FSynchronizedSessionData Data, const FStartTraceArgs& Args);
+	};
+}
+
+
+
diff --git a/Engine/Plugins/Developer/Concert/ConcertInsights/ConcertInsightsCore/Source/ConcertInsightsCore/Public/TraceMessages.h b/Engine/Plugins/Developer/Concert/ConcertInsights/ConcertInsightsCore/Source/ConcertInsightsCore/Public/TraceMessages.h
new file mode 100644
--- /dev/null
+++ b/Engine/Plugins/Developer/Concert/ConcertInsights/ConcertInsightsCore/Source/ConcertInsightsCore/Public/TraceMessages.h
@@ -0,0 +1,85 @@
+// Copyright Epic Games, Inc. All Rights Reserved.
+
+#pragma once
+
+#include "ProfilingDebugging/TraceAuxiliary.h"
+#include "TraceMessages.generated.h"
+
+/** See /ref FTraceAuxiliary::EConnectionType for more info. */
+UENUM()
+enum class EConcertTraceTargetType : uint8
+{
+	/** Connect to a trace server. Target is IP address or hostname. */
+	Network,
+	/**
+	 * Write to a file. Target string is filename. Absolute or relative current working directory.
+	 * If target is null the current date and time is used.
+	 */
+	File,
+};
+
+namespace UE::ConcertInsightsSync
+{
+	FORCEINLINE FTraceAuxiliary::EConnectionType ConvertTraceTargetType(EConcertTraceTargetType ConnectionType)
+	{
+		return static_cast<FTraceAuxiliary::EConnectionType>(ConnectionType);
+	}
+}
+
+
+/** Arguments to pass to FTraceAuxiliary::Start */
+USTRUCT()
+struct FConcertTrace_StartTraceArgs
+{
+	GENERATED_BODY()
+
+	/** How target is supposed to be interpreted. */
+	UPROPERTY()
+	EConcertTraceTargetType ConnectionType = EConcertTraceTargetType::Network;
+
+	/** String to use for connection. See /ref EConcertTraceTargetType for details */
+	UPROPERTY()
+	FString Target = TEXT("localhost");
+
+	/** Channels to use. Comma separated. */
+	UPROPERTY()
+	FString Channels = TEXT("default,Concert");
+};
+
+/**
+ * A request for the receive to:
+ * - start tracing,
+ * - ensure the .utrace contains a CONCERT_TRACE_INIT event
+ */
+USTRUCT()
+struct FConcertTrace_StartSyncTrace_Request
+{
+	GENERATED_BODY()
+
+	/** Info about how the trace is supposed to be started. */
+	UPROPERTY()
+	FConcertTrace_StartTraceArgs TraceArgs;
+};
+
+UENUM()
+enum class EConcertTraceErrorCode : uint8
+{
+	Timeout,
+	Joined,
+	Rejected
+};
+
+USTRUCT()
+struct FConcertTrace_StartSyncTrace_Response
+{
+	GENERATED_BODY()
+
+	UPROPERTY()
+	EConcertTraceErrorCode ErrorCode = EConcertTraceErrorCode::Timeout;
+};
+
+USTRUCT()
+struct FConcertTrace_StopSyncTrace
+{
+	GENERATED_BODY()
+};
\ No newline at end of file
diff --git a/Engine/Plugins/Developer/Concert/ConcertInsights/ConcertInsightsServer/ConcertInsightsServer.uplugin b/Engine/Plugins/Developer/Concert/ConcertInsights/ConcertInsightsServer/ConcertInsightsServer.uplugin
new file mode 100644
--- /dev/null
+++ b/Engine/Plugins/Developer/Concert/ConcertInsights/ConcertInsightsServer/ConcertInsightsServer.uplugin
@@ -0,0 +1,42 @@
+﻿{
+	"FileVersion": 3,
+	"Version": 1,
+	"VersionName": "1.0",
+	"FriendlyName": "ConcertInsightsServer",
+	"Description": "Listens for requests of clients to start synchronized tracing, which is initiated in the ConcertInsightsClient plugin.",
+	"Category": "Networking",
+	"CreatedBy": "Epic Games, Inc.",
+	"CreatedByURL": "",
+	"DocsURL": "",
+	"MarketplaceURL": "",
+	"CanContainContent": true,
+	"Hidden": true,
+	"IsBetaVersion": false,
+	"IsExperimentalVersion": true,
+	"Installed": false,
+	"SupportedPrograms": [
+		"UnrealMultiUserServer",
+		"UnrealMultiUserSlateServer"
+	],
+	"Modules": [
+		{
+			"Name": "ConcertInsightsServer",
+			"Type": "Program",
+			"LoadingPhase": "Default",
+			"ProgramAllowList": [
+				"UnrealMultiUserServer",
+				"UnrealMultiUserSlateServer" 
+			]
+		}
+	],
+	"Plugins": [
+		{
+			"Name": "ConcertInsightsCore",
+			"Enabled": true
+		},
+		{
+			"Name": "ConcertSyncServer",
+			"Enabled": true
+		}
+	]
+}
\ No newline at end of file
diff --git a/Engine/Plugins/Developer/Concert/ConcertInsights/ConcertInsightsServer/Source/ConcertInsightsServer/ConcertInsightsServer.Build.cs b/Engine/Plugins/Developer/Concert/ConcertInsights/ConcertInsightsServer/Source/ConcertInsightsServer/ConcertInsightsServer.Build.cs
new file mode 100644
--- /dev/null
+++ b/Engine/Plugins/Developer/Concert/ConcertInsights/ConcertInsightsServer/Source/ConcertInsightsServer/ConcertInsightsServer.Build.cs
@@ -0,0 +1,32 @@
+﻿// Copyright Epic Games, Inc. All Rights Reserved.
+
+using UnrealBuildTool;
+
+public class ConcertInsightsServer : ModuleRules
+{
+    public ConcertInsightsServer(ReadOnlyTargetRules Target) : base(Target)
+    {
+        PCHUsage = ModuleRules.PCHUsageMode.UseExplicitOrSharedPCHs;
+
+        PublicDependencyModuleNames.AddRange(
+            new string[]
+            {
+                "Core",
+            }
+        );
+
+        PrivateDependencyModuleNames.AddRange(
+            new string[]
+            {
+	            "ConcertInsightsCore", 
+	            "ConcertSyncServer", 
+                "CoreUObject",
+                "Slate",
+                "SlateCore",
+                "ToolMenus", 
+            }
+        );
+        
+        ShortName = "CrtInS";
+    }
+}
\ No newline at end of file
diff --git a/Engine/Plugins/Developer/Concert/ConcertInsights/ConcertInsightsServer/Source/ConcertInsightsServer/Private/ConcertInsightsServer.cpp b/Engine/Plugins/Developer/Concert/ConcertInsights/ConcertInsightsServer/Source/ConcertInsightsServer/Private/ConcertInsightsServer.cpp
new file mode 100644
--- /dev/null
+++ b/Engine/Plugins/Developer/Concert/ConcertInsights/ConcertInsightsServer/Source/ConcertInsightsServer/Private/ConcertInsightsServer.cpp
@@ -0,0 +1,21 @@
+﻿// Copyright Epic Games, Inc. All Rights Reserved.
+
+#include "ConcertInsightsServer.h"
+
+#define LOCTEXT_NAMESPACE "FConcertInsightsServerModule"
+
+namespace UE::ConcertInsightsServer
+{
+	void FConcertInsightsServerModule::StartupModule()
+	{
+		TraceControls = ConcertInsightsCore::FTraceControls::Make<FServerTraceControls>();
+	}
+
+	void FConcertInsightsServerModule::ShutdownModule()
+	{}
+}
+
+
+#undef LOCTEXT_NAMESPACE
+    
+IMPLEMENT_MODULE(UE::ConcertInsightsServer::FConcertInsightsServerModule, ConcertInsightsServer)
\ No newline at end of file
diff --git a/Engine/Plugins/Developer/Concert/ConcertInsights/ConcertInsightsServer/Source/ConcertInsightsServer/Private/ConcertInsightsServer.h b/Engine/Plugins/Developer/Concert/ConcertInsights/ConcertInsightsServer/Source/ConcertInsightsServer/Private/ConcertInsightsServer.h
new file mode 100644
--- /dev/null
+++ b/Engine/Plugins/Developer/Concert/ConcertInsights/ConcertInsightsServer/Source/ConcertInsightsServer/Private/ConcertInsightsServer.h
@@ -0,0 +1,29 @@
+﻿// Copyright Epic Games, Inc. All Rights Reserved.
+
+#pragma once
+
+#include "IConcertInsightsServerModule.h"
+#include "ServerTraceControls.h"
+#include "Templates/UniquePtr.h"
+
+namespace UE::ConcertInsightsServer
+{
+	class FConcertInsightsServerModule : public IConcertInsightsServerModule
+	{
+	public:
+		
+		//~ Begin IModuleInterface Interface
+		virtual void StartupModule() override;
+		virtual void ShutdownModule() override;
+		//~ End IModuleInterface Interface
+
+		FServerTraceControls& GetTraceControls() { return *TraceControls; }
+
+	private:
+
+		/** The local state of synchronized tracing. */
+		TUniquePtr<FServerTraceControls> TraceControls;
+	};
+}
+
+
diff --git a/Engine/Plugins/Developer/Concert/ConcertInsights/ConcertInsightsServer/Source/ConcertInsightsServer/Private/ServerTraceControls.cpp b/Engine/Plugins/Developer/Concert/ConcertInsights/ConcertInsightsServer/Source/ConcertInsightsServer/Private/ServerTraceControls.cpp
new file mode 100644
--- /dev/null
+++ b/Engine/Plugins/Developer/Concert/ConcertInsights/ConcertInsightsServer/Source/ConcertInsightsServer/Private/ServerTraceControls.cpp
@@ -0,0 +1,104 @@
+// Copyright Epic Games, Inc. All Rights Reserved.
+
+#include "ServerTraceControls.h"
+
+#include "IConcertServer.h"
+#include "IConcertSyncServer.h"
+#include "IConcertSyncServerModule.h"
+
+namespace UE::ConcertInsightsServer
+{
+	FServerTraceControls::FServerTraceControls()
+	{
+		IConcertSyncServerModule::Get().OnServerCreated().AddRaw(this, &FServerTraceControls::OnServerCreated);
+	}
+
+	FServerTraceControls::~FServerTraceControls()
+	{
+		if (IConcertSyncServerModule::IsAvailable())
+		{
+			IConcertSyncServerModule::Get().OnServerCreated().RemoveAll(this);
+		}
+
+		const TSharedPtr<IConcertServerSession> ServerSession = InProgressSynchronizedServerTrace ? InProgressSynchronizedServerTrace->InitiatingSession.Pin() : nullptr;
+		if (ServerSession)
+		{
+			ServerSession->OnSessionClientChanged().RemoveAll(this);
+		}
+
+		// ~FTraceControls will clean up delegates of registered sessions itself
+	}
+
+	void FServerTraceControls::OnSynchronizedTraceAccepted(const FConcertSessionContext& Context, const FConcertTrace_StartSyncTrace_Request& Request, const TSharedRef<IConcertSession>& Session)
+	{
+		// This is the equivalent of calling ServerInstance.Pin()->GetConcertServer()->GetLiveSession(Session->GetId())
+		TSharedRef<IConcertServerSession> ServerSession = StaticCastSharedRef<IConcertServerSession>(Session);
+		
+		check(!InProgressSynchronizedServerTrace.IsSet());
+		InProgressSynchronizedServerTrace.Emplace(Context.SourceEndpointId, ServerSession);
+		
+		ServerSession->OnSessionClientChanged().AddRaw(this, &FServerTraceControls::OnSynchronizedTraceClientChanged);
+		OnSynchronizedTraceStopped().AddRaw(this, &FServerTraceControls::CleanUpClientsChangedDelegate);
+	}
+
+	bool FServerTraceControls::CanSendRequestsToEndpoint(const FGuid& EndpointId, const IConcertSession& Session) const
+	{
+		return Session.GetSessionInfo().ServerEndpointId != EndpointId;
+	}
+
+	ConcertInsightsCore::FInitArgs FServerTraceControls::GetInitEventArgs() const
+	{
+		return { {}, TEXT("Server"), true };
+	}
+
+	void FServerTraceControls::OnServerCreated(TWeakPtr<IConcertSyncServer> Server)
+	{
+		if (!ensureMsgf(!ServerInstance.IsValid(), TEXT("We assume there is only one server instance per application.")))
+		{
+			return;
+		}
+		
+		if (const TSharedPtr<IConcertSyncServer> ServerPin = Server.Pin())
+		{
+			ServerInstance = Server;
+			RegisterHandlersForSessions(*ServerPin->GetConcertServer());
+		}
+	}
+
+	void FServerTraceControls::RegisterHandlersForSessions(IConcertServer& Server)
+	{
+		Server.OnConcertServerSessionStartup().AddRaw(this, &FServerTraceControls::OnSessionStartup);
+		for (const TSharedPtr<IConcertServerSession>& Session : Server.GetLiveSessions())
+		{
+			RegisterTraceRequestsHandler(Session.ToSharedRef());
+		}
+	}
+
+	void FServerTraceControls::OnSessionStartup(TWeakPtr<IConcertServerSession> Session)
+	{
+		if (const TSharedPtr<IConcertServerSession> SessionPin = Session.Pin())
+		{
+			RegisterTraceRequestsHandler(SessionPin.ToSharedRef());
+		}
+	}
+
+	void FServerTraceControls::OnSynchronizedTraceClientChanged(IConcertServerSession& Session, EConcertClientStatus Status, const FConcertSessionClientInfo& ClientInfo)
+	{
+		if (Status == EConcertClientStatus::Disconnected
+			&& InProgressSynchronizedServerTrace->SynchronizedTraceInstigator == ClientInfo.ClientEndpointId)
+		{
+			StopSynchronizedTrace();
+			CleanUpClientsChangedDelegate();
+		}
+	}
+
+	void FServerTraceControls::CleanUpClientsChangedDelegate()
+	{
+		if (const TSharedPtr<IConcertServerSession> ServerSession = InProgressSynchronizedServerTrace ? InProgressSynchronizedServerTrace->InitiatingSession.Pin() : nullptr)
+		{
+			ServerSession->OnSessionClientChanged().RemoveAll(this);
+		}
+
+		InProgressSynchronizedServerTrace.Reset();
+	}
+}
diff --git a/Engine/Plugins/Developer/Concert/ConcertInsights/ConcertInsightsServer/Source/ConcertInsightsServer/Private/ServerTraceControls.h b/Engine/Plugins/Developer/Concert/ConcertInsights/ConcertInsightsServer/Source/ConcertInsightsServer/Private/ServerTraceControls.h
new file mode 100644
--- /dev/null
+++ b/Engine/Plugins/Developer/Concert/ConcertInsights/ConcertInsightsServer/Source/ConcertInsightsServer/Private/ServerTraceControls.h
@@ -0,0 +1,67 @@
+// Copyright Epic Games, Inc. All Rights Reserved.
+
+#pragma once
+
+#include "TraceControls.h"
+
+#include "Containers/Array.h"
+#include "Templates/SharedPointer.h"
+
+class IConcertServer;
+class IConcertServerSession;
+class IConcertSyncServer;
+
+namespace UE::ConcertInsightsServer
+{
+	/**
+	 * Server controls for the editor.
+	 *
+	 * This class also keeps track of the client that instigated the synchronized trace.
+	 * If the instigating client disconnects, all other endpoints are told to stop tracing.
+	 */
+	class FServerTraceControls : public ConcertInsightsCore::FTraceControls
+	{
+		friend class FTraceControls; // For FTraceControls::Make
+	public:
+
+		virtual ~FServerTraceControls() override;
+		
+	protected:
+		
+		//~ Begin FTraceControls Interface
+		virtual void OnSynchronizedTraceAccepted(const FConcertSessionContext& Context, const FConcertTrace_StartSyncTrace_Request& Request, const TSharedRef<IConcertSession>& Session) override;
+		virtual bool CanSendRequestsToEndpoint(const FGuid& EndpointId, const IConcertSession& Session) const override;
+		virtual ConcertInsightsCore::FInitArgs GetInitEventArgs() const override;
+		//~ End FTraceControls Interface
+		
+	private:
+
+		struct FSynchronizedSessionServerData
+		{
+			/**
+			 * Endpoints of the endpoint that instigated the synchronized trace.
+			 * If this endpoint disconnects, the server tells all other clients to terminate the synchronized trace.
+			 */
+			FGuid SynchronizedTraceInstigator;
+
+			/** The server session on which this trace was started. */
+			TWeakPtr<IConcertServerSession> InitiatingSession;
+		};
+
+		/** Keeps track of the created server. Usually has 0 or 1 entries. */
+		TWeakPtr<IConcertSyncServer> ServerInstance;
+		
+		/** Additional data for the server about the synchronized trace. */
+		TOptional<FSynchronizedSessionServerData> InProgressSynchronizedServerTrace;
+		
+		FServerTraceControls();
+		
+		void OnServerCreated(TWeakPtr<IConcertSyncServer> Server);
+		void RegisterHandlersForSessions(IConcertServer& Server);
+		
+		void OnSessionStartup(TWeakPtr<IConcertServerSession> Session);
+		void OnSynchronizedTraceClientChanged(IConcertServerSession& ConcertServerSession, EConcertClientStatus ConcertClientStatus, const FConcertSessionClientInfo& ConcertSessionClientInfo);
+		
+		void CleanUpClientsChangedDelegate();
+	};
+}
diff --git a/Engine/Plugins/Developer/Concert/ConcertInsights/ConcertInsightsServer/Source/ConcertInsightsServer/Public/IConcertInsightsServerModule.h b/Engine/Plugins/Developer/Concert/ConcertInsights/ConcertInsightsServer/Source/ConcertInsightsServer/Public/IConcertInsightsServerModule.h
new file mode 100644
--- /dev/null
+++ b/Engine/Plugins/Developer/Concert/ConcertInsights/ConcertInsightsServer/Source/ConcertInsightsServer/Public/IConcertInsightsServerModule.h
@@ -0,0 +1,40 @@
+// Copyright Epic Games, Inc. All Rights Reserved.
+
+#pragma once
+
+#include "Modules/ModuleManager.h"
+
+namespace UE::ConcertInsightsServer
+{
+	/**
+	 * On the server, this module listens for requests to start synchronized session tracing.
+	 */
+	class CONCERTINSIGHTSSERVER_API IConcertInsightsServerModule : public IModuleInterface
+	{
+	public:
+
+		/**
+		 * Singleton-like access to this module's interface.  This is just for convenience!
+		 * Beware of calling this during the shutdown phase, though.  Your module might have been unloaded already.
+		 *
+		 * @return Returns singleton instance, loading the module on demand if needed
+		 */
+		static inline IConcertInsightsServerModule& Get()
+		{
+			static const FName ModuleName = "ConcertInsightsServer";
+			return FModuleManager::LoadModuleChecked<IConcertInsightsServerModule>(ModuleName);
+		}
+
+		/**
+		 * Checks to see if this module is loaded and ready.  It is only valid to call Get() during shutdown if IsAvailable() returns true.
+		 *
+		 * @return True if the module is loaded and ready to use
+		 */
+		static inline bool IsAvailable()
+		{
+			static const FName ModuleName = "ConcertInsightsEditor";
+			return FModuleManager::Get().IsModuleLoaded(ModuleName);
+		}
+	};
+}
+
diff --git a/Engine/Plugins/Developer/Concert/ConcertInsights/ConcertInsightsVisualizer/ConcertInsightsVisualizer.uplugin b/Engine/Plugins/Developer/Concert/ConcertInsights/ConcertInsightsVisualizer/ConcertInsightsVisualizer.uplugin
new file mode 100644
--- /dev/null
+++ b/Engine/Plugins/Developer/Concert/ConcertInsights/ConcertInsightsVisualizer/ConcertInsightsVisualizer.uplugin
@@ -0,0 +1,26 @@
+﻿{
+	"FileVersion": 3,
+	"Version": 1,
+	"VersionName": "1.0",
+	"FriendlyName": "ConcertInsightsVisualizer",
+	"Description": "Analyses and provides visualization widgets for Concert message types in Unreal Insights.",
+	"Category": "Networking",
+	"CreatedBy": "Epic Games, Inc.",
+	"CreatedByURL": "",
+	"DocsURL": "",
+	"MarketplaceURL": "",
+	"CanContainContent": false,
+	"Hidden": true,
+	"IsBetaVersion": false,
+	"IsExperimentalVersion": true,
+	"Installed": false,
+	"SupportedPrograms": [ "UnrealInsights" ],
+	"Modules": [
+		{
+			"Name": "ConcertInsightsVisualizer",
+			"Type": "EditorAndProgram",
+			"LoadingPhase": "Default",
+			"ProgramAllowList": [ "UnrealInsights" ]
+		}
+	]
+}
\ No newline at end of file
diff --git a/Engine/Plugins/Developer/Concert/ConcertInsights/ConcertInsightsVisualizer/Source/ConcertInsightsVisualizer/ConcertInsightsVisualizer.Build.cs b/Engine/Plugins/Developer/Concert/ConcertInsights/ConcertInsightsVisualizer/Source/ConcertInsightsVisualizer/ConcertInsightsVisualizer.Build.cs
new file mode 100644
--- /dev/null
+++ b/Engine/Plugins/Developer/Concert/ConcertInsights/ConcertInsightsVisualizer/Source/ConcertInsightsVisualizer/ConcertInsightsVisualizer.Build.cs
@@ -0,0 +1,31 @@
+﻿// Copyright Epic Games, Inc. All Rights Reserved.
+
+using UnrealBuildTool;
+
+public class ConcertInsightsVisualizer : ModuleRules
+{
+	public ConcertInsightsVisualizer(ReadOnlyTargetRules Target) : base(Target)
+	{
+		PCHUsage = ModuleRules.PCHUsageMode.UseExplicitOrSharedPCHs;
+		
+		PublicDependencyModuleNames.AddRange(
+			new string[]
+			{
+				"Core",
+			});
+		
+		PrivateDependencyModuleNames.AddRange(
+			new string[]
+			{
+				"CoreUObject",
+				"Projects",
+				"Slate",
+				"SlateCore",
+				"TraceAnalysis",
+				"TraceInsights",
+				"TraceServices"
+			});
+		
+		ShortName = "CrtInV";
+	}
+}
diff --git a/Engine/Plugins/Developer/Concert/ConcertInsights/ConcertInsightsVisualizer/Source/ConcertInsightsVisualizer/Private/ConcertInsightsStyle.cpp b/Engine/Plugins/Developer/Concert/ConcertInsights/ConcertInsightsVisualizer/Source/ConcertInsightsVisualizer/Private/ConcertInsightsStyle.cpp
new file mode 100644
--- /dev/null
+++ b/Engine/Plugins/Developer/Concert/ConcertInsights/ConcertInsightsVisualizer/Source/ConcertInsightsVisualizer/Private/ConcertInsightsStyle.cpp
@@ -0,0 +1,71 @@
+// Copyright Epic Games, Inc. All Rights Reserved.
+
+#include "ConcertInsightsStyle.h"
+
+#include "Brushes/SlateBoxBrush.h"
+#include "Brushes/SlateImageBrush.h"
+#include "Interfaces/IPluginManager.h"
+#include "Misc/Paths.h"
+#include "Styling/SlateStyle.h"
+#include "Styling/SlateStyleRegistry.h"
+#include "Styling/SlateTypes.h"
+#include "Styling/StarshipCoreStyle.h"
+
+LLM_DEFINE_TAG(ProtocolFrontendStyle);
+
+#define IMAGE_PLUGIN_BRUSH( RelativePath, ... ) FSlateImageBrush( FConcertInsightsStyle::InContent( RelativePath, ".png" ), __VA_ARGS__ )
+#define IMAGE_BRUSH(RelativePath, ...) FSlateImageBrush(StyleSet->RootToContentDir(RelativePath, TEXT(".png")), __VA_ARGS__)
+#define BOX_BRUSH(RelativePath, ... ) FSlateBoxBrush(StyleSet->RootToContentDir(RelativePath, TEXT(".png") ), __VA_ARGS__ )
+#define IMAGE_PLUGIN_BRUSH_SVG( RelativePath, ... ) FSlateVectorImageBrush(FConcertInsightsStyle::InContent(RelativePath, ".svg"), __VA_ARGS__)
+
+namespace UE::ConcertInsightsVisualizer
+{
+	FString FConcertInsightsStyle::InContent(const FString& RelativePath, const ANSICHAR* Extension)
+	{
+		static FString ContentDir = IPluginManager::Get().FindPlugin(TEXT("ConcertInsights"))->GetContentDir();
+		return (ContentDir / RelativePath) + Extension;
+	}
+
+	TSharedPtr<FSlateStyleSet > FConcertInsightsStyle::StyleSet;
+
+	FName FConcertInsightsStyle::GetStyleSetName()
+	{
+		return FName(TEXT("FConcertInsightsStyle"));
+	}
+
+	void FConcertInsightsStyle::Initialize()
+	{
+		LLM_SCOPE_BYTAG(ProtocolFrontendStyle);
+		// Only register once
+		if (StyleSet.IsValid())
+		{
+			return;
+		}
+
+		StyleSet = MakeShared<FSlateStyleSet>(GetStyleSetName());
+		StyleSet->SetContentRoot(FPaths::EngineContentDir() / TEXT("Editor/Slate"));
+		StyleSet->SetCoreContentRoot(FPaths::EngineContentDir() / TEXT("Slate"));
+		
+
+		FSlateStyleRegistry::RegisterSlateStyle(*StyleSet.Get());
+	};
+
+	void FConcertInsightsStyle::Shutdown()
+	{
+		if (StyleSet.IsValid())
+		{
+			FSlateStyleRegistry::UnRegisterSlateStyle(*StyleSet.Get());
+			ensure(StyleSet.IsUnique());
+			StyleSet.Reset();
+		}
+	}
+
+	TSharedPtr<class ISlateStyle> FConcertInsightsStyle::Get()
+	{
+		return StyleSet;
+	}
+}
+
+#undef IMAGE_PLUGIN_BRUSH
+#undef IMAGE_BRUSH
+
diff --git a/Engine/Plugins/Developer/Concert/ConcertInsights/ConcertInsightsVisualizer/Source/ConcertInsightsVisualizer/Private/ConcertInsightsStyle.h b/Engine/Plugins/Developer/Concert/ConcertInsights/ConcertInsightsVisualizer/Source/ConcertInsightsVisualizer/Private/ConcertInsightsStyle.h
new file mode 100644
--- /dev/null
+++ b/Engine/Plugins/Developer/Concert/ConcertInsights/ConcertInsightsVisualizer/Source/ConcertInsightsVisualizer/Private/ConcertInsightsStyle.h
@@ -0,0 +1,25 @@
+// Copyright Epic Games, Inc. All Rights Reserved.
+
+#pragma once
+
+#include "Styling/ISlateStyle.h"
+#include "Styling/SlateStyle.h"
+
+namespace UE::ConcertInsightsVisualizer
+{
+	class FConcertInsightsStyle
+	{
+	public:
+	
+		static void Initialize();
+		static void Shutdown();
+
+		static TSharedPtr<class ISlateStyle> Get();
+		static FName GetStyleSetName();
+	
+	private:
+	
+		static FString InContent(const FString& RelativePath, const ANSICHAR* Extension);
+		static TSharedPtr<FSlateStyleSet> StyleSet;
+	};
+}
\ No newline at end of file
diff --git a/Engine/Plugins/Developer/Concert/ConcertInsights/ConcertInsightsVisualizer/Source/ConcertInsightsVisualizer/Private/ConcertInsightsVisualizer.cpp b/Engine/Plugins/Developer/Concert/ConcertInsights/ConcertInsightsVisualizer/Source/ConcertInsightsVisualizer/Private/ConcertInsightsVisualizer.cpp
new file mode 100644
--- /dev/null
+++ b/Engine/Plugins/Developer/Concert/ConcertInsights/ConcertInsightsVisualizer/Source/ConcertInsightsVisualizer/Private/ConcertInsightsVisualizer.cpp
@@ -0,0 +1,44 @@
+﻿// Copyright Epic Games, Inc. All Rights Reserved.
+
+#include "ConcertInsightsVisualizer.h"
+
+#include "ConcertInsightsStyle.h"
+
+#include "Features/IModularFeatures.h"
+#include "Insights/ITimingViewExtender.h"
+
+#define LOCTEXT_NAMESPACE "FConcertInsightsVisualizerModule"
+
+namespace UE::ConcertInsightsVisualizer
+{
+	void FConcertInsightsVisualizerModule::StartupModule()
+	{
+		UE_LOG(LogTemp, Log, TEXT("Initing FConcertInsightsVisualizerModule..."));
+		
+		FConcertInsightsStyle::Initialize();
+		RegisterInsightsExtensions();
+	}
+
+	void FConcertInsightsVisualizerModule::ShutdownModule()
+	{
+		UnregisterInsightsExtensions();
+		FConcertInsightsStyle::Shutdown();
+	}
+
+	void FConcertInsightsVisualizerModule::RegisterInsightsExtensions()
+	{
+		IModularFeatures::Get().RegisterModularFeature(TraceServices::ModuleFeatureName, &ConcertInsightsModule);
+		IModularFeatures::Get().RegisterModularFeature(Insights::TimingViewExtenderFeatureName, &TimingViewExtender);
+	}
+
+	void FConcertInsightsVisualizerModule::UnregisterInsightsExtensions()
+	{
+		IModularFeatures::Get().UnregisterModularFeature(TraceServices::ModuleFeatureName, &ConcertInsightsModule);
+		IModularFeatures::Get().UnregisterModularFeature(Insights::TimingViewExtenderFeatureName, &TimingViewExtender);
+	}
+}
+
+
+#undef LOCTEXT_NAMESPACE
+	
+IMPLEMENT_MODULE(UE::ConcertInsightsVisualizer::FConcertInsightsVisualizerModule, ConcertInsightsVisualizer)
\ No newline at end of file
diff --git a/Engine/Plugins/Developer/Concert/ConcertInsights/ConcertInsightsVisualizer/Source/ConcertInsightsVisualizer/Private/ConcertInsightsVisualizer.h b/Engine/Plugins/Developer/Concert/ConcertInsights/ConcertInsightsVisualizer/Source/ConcertInsightsVisualizer/Private/ConcertInsightsVisualizer.h
new file mode 100644
--- /dev/null
+++ b/Engine/Plugins/Developer/Concert/ConcertInsights/ConcertInsightsVisualizer/Source/ConcertInsightsVisualizer/Private/ConcertInsightsVisualizer.h
@@ -0,0 +1,33 @@
+﻿// Copyright Epic Games, Inc. All Rights Reserved.
+
+#pragma once
+
+#include "IConcertInsightsVisualizerModule.h"
+#include "Extension/ConcertTimingViewExtender.h"
+#include "Extension/ConcertTraceInsightsModule.h"
+
+namespace UE::ConcertInsightsVisualizer
+{
+	class FConcertProfilerManager;
+	
+	class FConcertInsightsVisualizerModule : public IConcertInsightsVisualizerModule
+	{
+	public:
+
+		//~ Begin IModuleInterface Interface
+		virtual void StartupModule() override;
+		virtual void ShutdownModule() override;
+		//~ End IModuleInterface Interface
+
+	private:
+
+		/** Adds Concert analyzer and provider. */
+		FConcertTraceInsightsModule ConcertInsightsModule;
+		/** Adds tracks to the timing view as outlined in Trace/ConcertProtocolTrace.h. */
+		FConcertTimingViewExtender TimingViewExtender;
+
+		/** Registers all Concert specific modular features with Insights. */
+		void RegisterInsightsExtensions();
+		void UnregisterInsightsExtensions();
+	};
+}
diff --git a/Engine/Plugins/Developer/Concert/ConcertInsights/ConcertInsightsVisualizer/Source/ConcertInsightsVisualizer/Private/Extension/ConcertTimingViewExtender.cpp b/Engine/Plugins/Developer/Concert/ConcertInsights/ConcertInsightsVisualizer/Source/ConcertInsightsVisualizer/Private/Extension/ConcertTimingViewExtender.cpp
new file mode 100644
--- /dev/null
+++ b/Engine/Plugins/Developer/Concert/ConcertInsights/ConcertInsightsVisualizer/Source/ConcertInsightsVisualizer/Private/Extension/ConcertTimingViewExtender.cpp
@@ -0,0 +1,62 @@
+// Copyright Epic Games, Inc. All Rights Reserved.
+
+#include "ConcertTimingViewExtender.h"
+
+#include "Insights/ITimingViewSession.h"
+#include "Insights/IUnrealInsightsModule.h"
+
+namespace UE::ConcertInsightsVisualizer
+{
+	void FConcertTimingViewExtender::OnBeginSession(Insights::ITimingViewSession& InSession)
+	{
+		if (InSession.GetName() != FInsightsManagerTabs::TimingProfilerTabId)
+		{
+			return;
+		}
+
+		FPerSessionData* PerSessionData = PerSessionDataMap.Find(&InSession);
+		if (!PerSessionData)
+		{
+			PerSessionData = &PerSessionDataMap.Add(&InSession);
+			PerSessionData->SharedData = MakeUnique<FConcertTimingViewSession>();
+			PerSessionData->SharedData->OnBeginSession(InSession);
+		}
+		else
+		{
+			// Should generally not enter here because OnEndSession removes it but keep it in case Insights API skips a call (can that even happen?)
+			PerSessionData->SharedData->OnBeginSession(InSession);
+		}
+	}
+
+	void FConcertTimingViewExtender::OnEndSession(Insights::ITimingViewSession& InSession)
+	{
+		if (InSession.GetName() != FInsightsManagerTabs::TimingProfilerTabId)
+		{
+			return;
+		}
+
+		const FPerSessionData* PerSessionData = PerSessionDataMap.Find(&InSession);
+		if (PerSessionData)
+		{
+			PerSessionData->SharedData->OnEndSession(InSession);
+		}
+
+		PerSessionDataMap.Remove(&InSession);
+	}
+
+	void FConcertTimingViewExtender::Tick(Insights::ITimingViewSession& InTimingSession, const TraceServices::IAnalysisSession& InAnalysisSession)
+	{
+		if (const FPerSessionData* PerSessionData = PerSessionDataMap.Find(&InTimingSession))
+		{
+			PerSessionData->SharedData->Tick(InTimingSession, InAnalysisSession);
+		}
+	}
+
+	void FConcertTimingViewExtender::ExtendFilterMenu(Insights::ITimingViewSession& InSession, FMenuBuilder& InMenuBuilder)
+	{
+		if (const FPerSessionData* PerSessionData = PerSessionDataMap.Find(&InSession))
+		{
+			PerSessionData->SharedData->ExtendFilterMenu(InSession, InMenuBuilder);
+		}
+	}
+}
diff --git a/Engine/Plugins/Developer/Concert/ConcertInsights/ConcertInsightsVisualizer/Source/ConcertInsightsVisualizer/Private/Extension/ConcertTimingViewExtender.h b/Engine/Plugins/Developer/Concert/ConcertInsights/ConcertInsightsVisualizer/Source/ConcertInsightsVisualizer/Private/Extension/ConcertTimingViewExtender.h
new file mode 100644
--- /dev/null
+++ b/Engine/Plugins/Developer/Concert/ConcertInsights/ConcertInsightsVisualizer/Source/ConcertInsightsVisualizer/Private/Extension/ConcertTimingViewExtender.h
@@ -0,0 +1,35 @@
+// Copyright Epic Games, Inc. All Rights Reserved.
+
+#pragma once
+
+#include "ConcertTimingViewSession.h"
+
+#include "Insights/ITimingViewExtender.h"
+#include "Templates/UniquePtr.h"
+
+namespace UE::ConcertInsightsVisualizer
+{
+	/** Keeps track of FConcertTimingViewSession per analytics session. */
+	class FConcertTimingViewExtender : public Insights::ITimingViewExtender
+	{
+	public:
+
+		//~ Begin Insights::ITimingViewExtender Interface
+		virtual void OnBeginSession(Insights::ITimingViewSession& InSession) override;
+		virtual void OnEndSession(Insights::ITimingViewSession& InSession) override;
+		virtual void Tick(Insights::ITimingViewSession& InTimingSession, const TraceServices::IAnalysisSession& InAnalysisSession) override;
+		virtual void ExtendFilterMenu(Insights::ITimingViewSession& InSession, FMenuBuilder& InMenuBuilder) override;
+		//~ End Insights::ITimingViewExtender Interface
+
+	private:
+
+		struct FPerSessionData
+		{
+			TUniquePtr<FConcertTimingViewSession> SharedData;
+		};
+
+		//** The data we host per-session */
+		TMap<Insights::ITimingViewSession*, FPerSessionData> PerSessionDataMap;
+	};
+}
+
diff --git a/Engine/Plugins/Developer/Concert/ConcertInsights/ConcertInsightsVisualizer/Source/ConcertInsightsVisualizer/Private/Extension/ConcertTimingViewSession.cpp b/Engine/Plugins/Developer/Concert/ConcertInsights/ConcertInsightsVisualizer/Source/ConcertInsightsVisualizer/Private/Extension/ConcertTimingViewSession.cpp
new file mode 100644
--- /dev/null
+++ b/Engine/Plugins/Developer/Concert/ConcertInsights/ConcertInsightsVisualizer/Source/ConcertInsightsVisualizer/Private/Extension/ConcertTimingViewSession.cpp
@@ -0,0 +1,80 @@
+// Copyright Epic Games, Inc. All Rights Reserved.
+
+#include "ConcertTimingViewSession.h"
+
+#include "Trace/TraceAggregator.h"
+#include "Track/ProtocolTrack.h"
+
+#include "Framework/MultiBox/MultiBoxBuilder.h"
+#include "Insights/ITimingViewSession.h"
+#include "Insights/IUnrealInsightsModule.h"
+#include "Modules/ModuleManager.h"
+#include "Trace/ProtocolMultiEndpointProvider.h"
+
+#define LOCTEXT_NAMESPACE "FConcertTimingViewSession"
+
+namespace UE::ConcertInsightsVisualizer
+{
+	void FConcertTimingViewSession::OnBeginSession(Insights::ITimingViewSession& InSession)
+	{}
+
+	void FConcertTimingViewSession::OnEndSession(Insights::ITimingViewSession& InSession)
+	{}
+	
+	void FConcertTimingViewSession::Tick(Insights::ITimingViewSession& InTimingSession, const TraceServices::IAnalysisSession& InAnalysisSession)
+	{
+		const FProtocolMultiEndpointProvider* ProtocolProvider = InAnalysisSession.ReadProvider<FProtocolMultiEndpointProvider>(FProtocolMultiEndpointProvider::ProviderName);
+		if (!ensure(ProtocolProvider))
+		{
+			return;
+		}
+		
+		if (!ObjectTrack)
+		{
+			ObjectTrack = MakeShared<FProtocolTrack>(InAnalysisSession, *ProtocolProvider);
+			ObjectTrack->SetVisibilityFlag(true);
+			ObjectTrack->SetOrder(FTimingTrackOrder::First);
+			InTimingSession.AddScrollableTrack(ObjectTrack);
+			InTimingSession.InvalidateScrollableTracksOrder();
+		}
+	}
+
+	void FConcertTimingViewSession::ExtendFilterMenu(Insights::ITimingViewSession& InSession, FMenuBuilder& InMenuBuilder)
+	{
+		if (!ObjectTrack)
+		{
+			return;
+		}
+		
+		InMenuBuilder.BeginSection("ConcertTracks", LOCTEXT("ConcertTracksSection", "Concert"));
+		{
+			InMenuBuilder.AddMenuEntry(
+				LOCTEXT("TrimPackagePaths.Label", "Show full package paths"),
+				LOCTEXT("TrimPackagePaths.Tooltip", "Whether the 1st line in every sequence should show the full path to the object."),
+				FSlateIcon(),
+				FUIAction(
+					FExecuteAction::CreateLambda([this](){ ObjectTrack->ToggleShowObjectFullPaths(); }),
+					FCanExecuteAction::CreateLambda([](){ return true; }),
+					FIsActionChecked::CreateLambda([this](){ return ObjectTrack->ShouldShowFullObjectPaths(); })
+					),
+				NAME_None,
+				EUserInterfaceActionType::ToggleButton
+			);
+		}
+		InMenuBuilder.EndSection();
+	}
+	
+	void FConcertTimingViewSession::ConditionalInitTraceAggregator(const TraceServices::IAnalysisSession& InAnalysisSession)
+	{
+		if (!TraceAggregator)
+		{
+			IUnrealInsightsModule& InsightsModule = FModuleManager::LoadModuleChecked<IUnrealInsightsModule>("TraceInsights");
+			Trace::FStoreClient* StoreClient = InsightsModule.GetStoreClient();
+			
+			TraceAggregator = MakeShared<FTraceAggregator>(*StoreClient, InAnalysisSession.GetTraceId());
+			TraceAggregator->StartAggregatedAnalysis();
+		}
+	}
+}
+
+#undef LOCTEXT_NAMESPACE
\ No newline at end of file
diff --git a/Engine/Plugins/Developer/Concert/ConcertInsights/ConcertInsightsVisualizer/Source/ConcertInsightsVisualizer/Private/Extension/ConcertTimingViewSession.h b/Engine/Plugins/Developer/Concert/ConcertInsights/ConcertInsightsVisualizer/Source/ConcertInsightsVisualizer/Private/Extension/ConcertTimingViewSession.h
new file mode 100644
--- /dev/null
+++ b/Engine/Plugins/Developer/Concert/ConcertInsights/ConcertInsightsVisualizer/Source/ConcertInsightsVisualizer/Private/Extension/ConcertTimingViewSession.h
@@ -0,0 +1,39 @@
+// Copyright Epic Games, Inc. All Rights Reserved.
+
+#pragma once
+
+#include "Templates/SharedPointer.h"
+#include "Templates/UnrealTemplate.h"
+
+class FMenuBuilder;
+
+namespace Insights { class ITimingViewSession; }
+namespace TraceServices { class IAnalysisSession; }
+
+namespace UE::ConcertInsightsVisualizer
+{
+	class FProtocolTrack;
+	class FTraceAggregator;
+
+	/** Adds tracks to the timing view as outlined in Trace/ConcertProtocolTrace.h. */
+	class FConcertTimingViewSession : public FNoncopyable
+	{
+	public:
+
+		// Mirrors the ITimingViewExtender interface - called by FConcertTimingViewExtender
+		void OnBeginSession(Insights::ITimingViewSession& InSession);
+		void OnEndSession(Insights::ITimingViewSession& InSession);
+		void Tick(Insights::ITimingViewSession& InTimingSession, const TraceServices::IAnalysisSession& InAnalysisSession);
+		void ExtendFilterMenu(Insights::ITimingViewSession& InSession, FMenuBuilder& InMenuBuilder);
+
+	private:
+
+		/** Generates the rows in the Insights UI. */
+		TSharedPtr<FProtocolTrack> ObjectTrack;
+
+		/** Traces related files and exposes their data. */
+		TSharedPtr<FTraceAggregator> TraceAggregator;
+		
+		void ConditionalInitTraceAggregator(const TraceServices::IAnalysisSession& InAnalysisSession);
+	};
+}
diff --git a/Engine/Plugins/Developer/Concert/ConcertInsights/ConcertInsightsVisualizer/Source/ConcertInsightsVisualizer/Private/Extension/ConcertTraceInsightsModule.cpp b/Engine/Plugins/Developer/Concert/ConcertInsights/ConcertInsightsVisualizer/Source/ConcertInsightsVisualizer/Private/Extension/ConcertTraceInsightsModule.cpp
new file mode 100644
--- /dev/null
+++ b/Engine/Plugins/Developer/Concert/ConcertInsights/ConcertInsightsVisualizer/Source/ConcertInsightsVisualizer/Private/Extension/ConcertTraceInsightsModule.cpp
@@ -0,0 +1,27 @@
+// Copyright Epic Games, Inc. All Rights Reserved.
+
+#include "ConcertTraceInsightsModule.h"
+
+#include "LogConcertInsights.h"
+#include "Trace/Analysis/ProtocolEndpointAnalyzer.h"
+#include "Trace/ProtocolMultiEndpointProvider.h"
+
+namespace UE::ConcertInsightsVisualizer
+{
+	void FConcertTraceInsightsModule::GetModuleInfo(TraceServices::FModuleInfo& OutModuleInfo)
+	{
+		OutModuleInfo.Name = TEXT("ConcertTrace");
+		OutModuleInfo.DisplayName = TEXT("Concert");
+	}
+
+	void FConcertTraceInsightsModule::OnAnalysisBegin(TraceServices::IAnalysisSession& Session)
+	{
+		UE_LOG(LogConcertInsights, Log, TEXT("FConcertTraceInsightsModule::OnAnalysisBegin"));
+
+		const TSharedRef<FProtocolMultiEndpointProvider> Provider = MakeShared<FProtocolMultiEndpointProvider>(Session);
+		Session.AddProvider(FProtocolMultiEndpointProvider::ProviderName, Provider);
+
+		// Insights takes ownership and will call delete 
+		Session.AddAnalyzer(new FProtocolEndpointAnalyzer(Session, *Provider));
+	}
+}
diff --git a/Engine/Plugins/Developer/Concert/ConcertInsights/ConcertInsightsVisualizer/Source/ConcertInsightsVisualizer/Private/Extension/ConcertTraceInsightsModule.h b/Engine/Plugins/Developer/Concert/ConcertInsights/ConcertInsightsVisualizer/Source/ConcertInsightsVisualizer/Private/Extension/ConcertTraceInsightsModule.h
new file mode 100644
--- /dev/null
+++ b/Engine/Plugins/Developer/Concert/ConcertInsights/ConcertInsightsVisualizer/Source/ConcertInsightsVisualizer/Private/Extension/ConcertTraceInsightsModule.h
@@ -0,0 +1,19 @@
+// Copyright Epic Games, Inc. All Rights Reserved.
+
+#pragma once
+
+#include "TraceServices/ModuleService.h"
+
+namespace UE::ConcertInsightsVisualizer
+{
+	/** Adds the Concert analyzer and provider when a session starts. */
+	class FConcertTraceInsightsModule : public TraceServices::IModule
+	{
+	public:
+
+		//~ Begin TraceServices::IModule Interface
+		virtual void GetModuleInfo(TraceServices::FModuleInfo& OutModuleInfo) override;
+		virtual void OnAnalysisBegin(TraceServices::IAnalysisSession& Session) override;
+		//~ End TraceServices::IModule Interface
+	};
+}
\ No newline at end of file
diff --git a/Engine/Plugins/Developer/Concert/ConcertInsights/ConcertInsightsVisualizer/Source/ConcertInsightsVisualizer/Private/Extension/Track/ProtocolTrack.cpp b/Engine/Plugins/Developer/Concert/ConcertInsights/ConcertInsightsVisualizer/Source/ConcertInsightsVisualizer/Private/Extension/Track/ProtocolTrack.cpp
new file mode 100644
--- /dev/null
+++ b/Engine/Plugins/Developer/Concert/ConcertInsights/ConcertInsightsVisualizer/Source/ConcertInsightsVisualizer/Private/Extension/Track/ProtocolTrack.cpp
@@ -0,0 +1,127 @@
+// Copyright Epic Games, Inc. All Rights Reserved.
+
+#include "ProtocolTrack.h"
+
+#include "TraceServices/Containers/Timelines.h"
+#include "Trace/ProtocolMultiEndpointProvider.h"
+
+namespace UE::ConcertInsightsVisualizer
+{
+	FProtocolTrack::FProtocolTrack(const TraceServices::IAnalysisSession& Session, const FProtocolMultiEndpointProvider& Provider)
+		: FTimingEventsTrack(TEXT("Concert"))
+		, Session(Session)
+		, Provider(Provider)
+	{}
+
+	void FProtocolTrack::ToggleShowObjectFullPaths()
+	{
+		bShouldShowFullObjectPaths = !bShouldShowFullObjectPaths;
+		SetDirtyFlag();
+	}
+
+	void FProtocolTrack::BuildDrawState(ITimingEventsTrackDrawStateBuilder& Builder, const ITimingTrackUpdateContext& Context)
+	{
+		TraceServices::FAnalysisSessionReadScope ReadScope(Session);
+		
+		FBuildContext ContextInternal(Builder, Context);
+		BuildSequences(ContextInternal);
+	}
+
+	void FProtocolTrack::InitTooltip(FTooltipDrawState& InOutTooltip, const ITimingEvent& InTooltipEvent) const
+	{
+		FTimingEventsTrack::InitTooltip(InOutTooltip, InTooltipEvent);
+	}
+
+	const TSharedPtr<const ITimingEvent> FProtocolTrack::SearchEvent(const FTimingEventSearchParameters& InSearchParameters) const
+	{
+		return FTimingEventsTrack::SearchEvent(InSearchParameters);
+	}
+
+	void FProtocolTrack::BuildSequences(FBuildContext& Context) const
+	{
+		// TODO: We should better visually separate protocols and objects, maybe by using separate tracks
+		Provider.EnumerateProtocols([&](FProtocolId ProtocolId)
+		{
+			Provider.EnumerateObjects(ProtocolId, [&](FObjectPath ObjectPath)
+			{
+				const FObjectScopeInfo Object { ProtocolId, ObjectPath };
+				Provider.EnumerateSequences(Context.StartTime, Context.EndTime, Object, [&](FSequenceId SequenceId)
+				{
+					BuildSequence(Context, { ProtocolId, ObjectPath, SequenceId });
+					return TraceServices::EEventEnumerate::Continue;
+				});
+				return TraceServices::EEventEnumerate::Continue;
+			});
+			return TraceServices::EEventEnumerate::Continue;
+		});
+	}
+
+	void FProtocolTrack::BuildSequence(FBuildContext& Context, const FSequenceScopeInfo& Info) const
+	{
+		// 1st row shows the sequence ID and
+		const TOptional<FVector2d> Bounds = Provider.GetSequenceBounds(Info);
+		Context.Builder.AddEvent(Bounds->X, Bounds->Y, Context.CurrentDepthOffset, *FString::Printf(TEXT("%u - %s"), Info.SequenceId, *GetObjectDisplayString(Info.ObjectPath.GetData())));
+		++Context.CurrentDepthOffset;
+
+		// 2nd row shows e.g. "Client Name" or "Transmission"
+		Provider.EnumerateNetworkScopes(Context.StartTime, Context.EndTime, Info, [&](double Start, double End, const FObjectNetworkScope& NetworkScope)
+		{
+			const FString EventName = NetworkScope.ProcessingEndpoint
+				? Provider.GetEndpointDisplayName(*NetworkScope.ProcessingEndpoint)
+				: TEXT("Transmission");
+			Context.Builder.AddEvent(Start, End, Context.CurrentDepthOffset, *EventName);
+			return TraceServices::EEventEnumerate::Continue;
+		});
+		++Context.CurrentDepthOffset;
+
+		// 3rd row and after show the where the endpoints spent CPU time
+		uint32 MaxNumRows = 0;
+		Provider.EnumerateEndpointsInSequence(Info, [&](FEndpointId EndpointId)
+		{
+			const uint32 NumRows = BuildCpuTimeline(Context, Info.MakeEndpointInfo(EndpointId));
+			MaxNumRows = FMath::Max(MaxNumRows, NumRows);
+			return TraceServices::EEventEnumerate::Continue;
+		});
+		Context.CurrentDepthOffset += MaxNumRows;
+	}
+
+	uint32 FProtocolTrack::BuildCpuTimeline(FBuildContext& Context, const FEndpointScopeInfo& Info) const
+	{
+		uint32 CpuTimelineNumRows = 0;
+		Provider.ReadProcessingStepTimeline(Info, [&](const TraceServices::ITimeline<FObjectProcessingStep>& Timeline)
+		{
+			Timeline.EnumerateEventsDownSampled(Context.StartTime, Context.EndTime, Context.SecondsPerPixel,
+				[&](double StartTime, double EndTime, uint32 Depth, const FObjectProcessingStep& Event)
+				{
+					const uint32 FinalOffset = Context.CurrentDepthOffset + Depth;
+					// The top-most row is at depth 0, so +1 since we're counting the number of rows
+					CpuTimelineNumRows = FMath::Max(CpuTimelineNumRows, Depth + 1); 
+					
+					Context.Builder.AddEvent(StartTime, EndTime, FinalOffset, Event.EventName);
+					return TraceServices::EEventEnumerate::Continue;
+				});
+		});
+
+		return CpuTimelineNumRows;
+	}
+
+	FString FProtocolTrack::GetObjectDisplayString(const TCHAR* ObjectPath) const
+	{
+		FString ObjectPathAsString = ObjectPath;
+		if (bShouldShowFullObjectPaths)
+		{
+			return ObjectPathAsString;
+		}
+
+		const int32 FoundPosition = ObjectPathAsString.Find(TEXT("PersistentLevel."));
+		if (FoundPosition == INDEX_NONE)
+		{
+			return ObjectPathAsString;
+		}
+
+		// PersistentLevel. has 16 characters
+		const int32 ChopPosition = FoundPosition + 16;
+		ObjectPathAsString.RightChopInline(ChopPosition);
+		return ObjectPathAsString;
+	}
+}
diff --git a/Engine/Plugins/Developer/Concert/ConcertInsights/ConcertInsightsVisualizer/Source/ConcertInsightsVisualizer/Private/Extension/Track/ProtocolTrack.h b/Engine/Plugins/Developer/Concert/ConcertInsights/ConcertInsightsVisualizer/Source/ConcertInsightsVisualizer/Private/Extension/Track/ProtocolTrack.h
new file mode 100644
--- /dev/null
+++ b/Engine/Plugins/Developer/Concert/ConcertInsights/ConcertInsightsVisualizer/Source/ConcertInsightsVisualizer/Private/Extension/Track/ProtocolTrack.h
@@ -0,0 +1,69 @@
+// Copyright Epic Games, Inc. All Rights Reserved.
+
+#pragma once
+
+#include "Insights/ViewModels/TimingEventsTrack.h"
+#include "Insights/ViewModels/TimingTrackViewport.h"
+#include "Trace/ProtocolMultiEndpointProvider.h"
+
+namespace UE::ConcertInsightsVisualizer
+{
+	class FProtocolGraphSeries;
+
+	class FProtocolTrack : public FTimingEventsTrack
+	{
+	public:
+
+		FProtocolTrack(const TraceServices::IAnalysisSession& Session UE_LIFETIMEBOUND, const FProtocolMultiEndpointProvider& Provider UE_LIFETIMEBOUND);
+
+		/** Toggles whether objects are displayed with full object paths or just the bit after PersistentLevel. */
+		void ToggleShowObjectFullPaths();
+		/** @return Whether objects are displayed with full object paths or just the bit after PersistentLevel */
+		bool ShouldShowFullObjectPaths() const { return bShouldShowFullObjectPaths; }
+		
+		//~ Begin FTimingEventsTrack Interface
+		virtual void BuildDrawState(ITimingEventsTrackDrawStateBuilder& Builder, const ITimingTrackUpdateContext& Context) override;
+		virtual void InitTooltip(FTooltipDrawState& InOutTooltip, const ITimingEvent& InTooltipEvent) const override;
+		virtual const TSharedPtr<const ITimingEvent> SearchEvent(const FTimingEventSearchParameters& InSearchParameters) const override;
+		//~ End FTimingEventsTrack Interface
+
+	private:
+
+		struct FBuildContext
+		{
+			ITimingEventsTrackDrawStateBuilder& Builder;
+			const ITimingTrackUpdateContext& TimingContext;
+			
+			const double StartTime = GetViewport().GetStartTime();
+			const double EndTime = GetViewport().GetEndTime();
+			const double SecondsPerPixel = 1.0 / GetViewport().GetScaleX();
+			
+			uint32 CurrentDepthOffset = 0;
+
+			FBuildContext(ITimingEventsTrackDrawStateBuilder& Builder, const ITimingTrackUpdateContext& Context)
+				: Builder(Builder)
+				, TimingContext(Context)
+			{}
+
+			const FTimingTrackViewport& GetViewport() const { return TimingContext.GetViewport(); }
+		};
+
+		const TraceServices::IAnalysisSession& Session;
+		const FProtocolMultiEndpointProvider& Provider;
+
+		/** Whether the 1st row of every sequence should show the full object path or just the part behind .PersistentLevel */
+		bool bShouldShowFullObjectPaths = false;
+
+		void BuildSequences(FBuildContext& Context) const;
+		void BuildSequence(FBuildContext& Context, const FSequenceScopeInfo& Info) const;
+		uint32 BuildCpuTimeline(
+			FBuildContext& Context,
+			const FEndpointScopeInfo& Info
+			) const;
+
+		/** Gets the final display string of ObjectPath (given the value of bShouldShowFullObjectPaths). */
+		FString GetObjectDisplayString(const TCHAR* ObjectPath) const;
+	};
+}
+
+
diff --git a/Engine/Plugins/Developer/Concert/ConcertInsights/ConcertInsightsVisualizer/Source/ConcertInsightsVisualizer/Private/LogConcertInsights.cpp b/Engine/Plugins/Developer/Concert/ConcertInsights/ConcertInsightsVisualizer/Source/ConcertInsightsVisualizer/Private/LogConcertInsights.cpp
new file mode 100644
--- /dev/null
+++ b/Engine/Plugins/Developer/Concert/ConcertInsights/ConcertInsightsVisualizer/Source/ConcertInsightsVisualizer/Private/LogConcertInsights.cpp
@@ -0,0 +1,5 @@
+// Copyright Epic Games, Inc. All Rights Reserved.
+
+#include "LogConcertInsights.h"
+
+DEFINE_LOG_CATEGORY(LogConcertInsights);
\ No newline at end of file
diff --git a/Engine/Plugins/Developer/Concert/ConcertInsights/ConcertInsightsVisualizer/Source/ConcertInsightsVisualizer/Private/LogConcertInsights.h b/Engine/Plugins/Developer/Concert/ConcertInsights/ConcertInsightsVisualizer/Source/ConcertInsightsVisualizer/Private/LogConcertInsights.h
new file mode 100644
--- /dev/null
+++ b/Engine/Plugins/Developer/Concert/ConcertInsights/ConcertInsightsVisualizer/Source/ConcertInsightsVisualizer/Private/LogConcertInsights.h
@@ -0,0 +1,7 @@
+// Copyright Epic Games, Inc. All Rights Reserved.
+
+#pragma once
+
+#include "Logging/LogMacros.h"
+
+DECLARE_LOG_CATEGORY_EXTERN(LogConcertInsights, Log, All);
\ No newline at end of file
diff --git a/Engine/Plugins/Developer/Concert/ConcertInsights/ConcertInsightsVisualizer/Source/ConcertInsightsVisualizer/Private/Tests/TimeSyncUtils.spec.cpp b/Engine/Plugins/Developer/Concert/ConcertInsights/ConcertInsightsVisualizer/Source/ConcertInsightsVisualizer/Private/Tests/TimeSyncUtils.spec.cpp
new file mode 100644
--- /dev/null
+++ b/Engine/Plugins/Developer/Concert/ConcertInsights/ConcertInsightsVisualizer/Source/ConcertInsightsVisualizer/Private/Tests/TimeSyncUtils.spec.cpp
@@ -0,0 +1,64 @@
+// Copyright Epic Games, Inc. All Rights Reserved.
+
+#include "Misc/AutomationTest.h"
+
+#include "Util/TimeSyncUtils.h"
+
+#include "Misc/AutomationTest.h"
+#include "Misc/Paths.h"
+#include "HAL/FileManager.h"
+
+namespace UE::ConcertInsightsVisualizer::Tests
+{
+	
+	BEGIN_DEFINE_SPEC(FTimeSyncUtilsSpec, "Editor.Concert.Insights", EAutomationTestFlags::EditorContext | EAutomationTestFlags::EngineFilter)
+	
+	END_DEFINE_SPEC(FTimeSyncUtilsSpec);
+
+	void FTimeSyncUtilsSpec::Define()
+	{
+		BeforeEach([this]()
+		{
+			
+		});
+
+		Describe("Convert time correctly", [this]()
+		{
+			/* See documentation of ConvertSourceToTargetTime
+			 * Global	0123456789_
+			 * Target	[Init567x9]
+			 * Source	-[1Init6y8]
+			 */
+			It("For example documentation case", [this]()
+			{
+				const FDateTime TargetInitUtc(2024, 2, 22, 12, 0, 1);
+				const FDateTime SourceInitUtc(2024, 2, 22, 12, 0, 3);
+				constexpr double TargetInitTime = 1;
+				constexpr double SourceInitTime = 2;
+				constexpr double SourceTime = 7; // = y from above
+
+				// Compute x from tomorrow
+				const double ConvertedTime = TimeSyncUtils::ConvertSourceToTargetTime(
+					TargetInitUtc, SourceInitUtc, TargetInitTime, SourceInitTime, SourceTime
+					);
+				// Inverse: Compute y from x
+				const double InverseConvertedTime = TimeSyncUtils::ConvertSourceToTargetTime(
+					SourceInitUtc, TargetInitUtc, SourceInitTime, TargetInitTime, ConvertedTime
+					);
+				
+				TestEqual(TEXT("Docu use case time"), ConvertedTime, 8.0);
+				TestEqual(TEXT("Inverse docu use case time"), InverseConvertedTime, SourceTime);
+			});
+			
+			It("When target 30 seconds later", [this]()
+			{
+
+			});
+
+			It("When target 30 seconds earlier", [this]()
+			{
+				
+			});
+		});
+	}
+}
diff --git a/Engine/Plugins/Developer/Concert/ConcertInsights/ConcertInsightsVisualizer/Source/ConcertInsightsVisualizer/Private/Trace/Analysis/IProtocolDataTarget.h b/Engine/Plugins/Developer/Concert/ConcertInsights/ConcertInsightsVisualizer/Source/ConcertInsightsVisualizer/Private/Trace/Analysis/IProtocolDataTarget.h
new file mode 100644
--- /dev/null
+++ b/Engine/Plugins/Developer/Concert/ConcertInsights/ConcertInsightsVisualizer/Source/ConcertInsightsVisualizer/Private/Trace/Analysis/IProtocolDataTarget.h
@@ -0,0 +1,36 @@
+// Copyright Epic Games, Inc. All Rights Reserved.
+
+#pragma once
+
+namespace UE::ConcertInsightsVisualizer
+{
+	struct FInitMessage;
+	struct FObjectSinkMessage;
+	struct FObjectTraceBeginMessage;
+	struct FObjectTraceEndMessage;
+	struct FObjectTransmissionReceiveMessage;
+	struct FObjectTransmissionStartMessage;
+	
+	/**
+	 * Target for receiving data from FProtocolEndpointAnalyzer.
+	 * 
+	 * This abstracts what is done with the data, which is important for aggregation:
+	 * - The main .utrace, which is the one that is technically open, is processed on the main thread and directly inserted into FProtocolMultiEndpointProvider.
+	 * - The aggregated .utrace files are processed on a separate thread. The analyzed data must be synchronized with the main thread which is controlled by
+	 * FProtocolDataQueue. FProtocolMultiEndpointProvider will dequeue ready data from FProtocolDataQueue on Tick.
+	 * 
+	 */
+	class IProtocolDataTarget
+	{
+	public:
+		
+		virtual void AppendInit(FInitMessage Message) = 0;
+		virtual void AppendObjectTraceBegin(FObjectTraceBeginMessage Message) = 0;
+		virtual void AppendObjectTraceEnd(FObjectTraceEndMessage Message) = 0;
+		virtual void AppendObjectTransmissionStart(FObjectTransmissionStartMessage Message) = 0;
+		virtual void AppendObjectTransmissionReceive(FObjectTransmissionReceiveMessage Message) = 0;
+		virtual void AppendObjectSink(FObjectSinkMessage Message) = 0;
+
+		virtual ~IProtocolDataTarget() = default;
+	};
+}
\ No newline at end of file
diff --git a/Engine/Plugins/Developer/Concert/ConcertInsights/ConcertInsightsVisualizer/Source/ConcertInsightsVisualizer/Private/Trace/Analysis/ProtocolDataQueue.h b/Engine/Plugins/Developer/Concert/ConcertInsights/ConcertInsightsVisualizer/Source/ConcertInsightsVisualizer/Private/Trace/Analysis/ProtocolDataQueue.h
new file mode 100644
--- /dev/null
+++ b/Engine/Plugins/Developer/Concert/ConcertInsights/ConcertInsightsVisualizer/Source/ConcertInsightsVisualizer/Private/Trace/Analysis/ProtocolDataQueue.h
@@ -0,0 +1,77 @@
+// Copyright Epic Games, Inc. All Rights Reserved.
+
+#pragma once
+
+#include "IProtocolDataTarget.h"
+#include "Trace/Messages/EMessageType.h"
+#include "Trace/Messages/InitMessage.h"
+#include "Trace/Messages/ObjectSinkMessage.h"
+#include "Trace/Messages/ObjectTraceMessage.h"
+#include "Trace/Messages/ObjectTransmissionReceiveMessage.h"
+#include "Trace/Messages/ObjectTransmissionStartMessage.h"
+
+#include "Containers/Queue.h"
+
+#include <type_traits>
+
+namespace UE::ConcertInsightsVisualizer
+{
+	static_assert(std::is_trivial_v<FInitMessage>, "FInitMessage must be trivial for the FProtocolQueuedItem::FMessage union.");
+	static_assert(std::is_trivial_v<FObjectTraceBeginMessage>, "FObjectTraceBeginMessage must be trivial for the FProtocolQueuedItem::FMessage union.");
+	static_assert(std::is_trivial_v<FObjectTraceEndMessage>, "FObjectTraceEndMessage must be trivial for the FProtocolQueuedItem::FMessage union.");
+	static_assert(std::is_trivial_v<FObjectTransmissionStartMessage>, "FObjectTransmissionStartMessage must be trivial for the FProtocolQueuedItem::FMessage union.");
+	static_assert(std::is_trivial_v<FObjectTransmissionReceiveMessage>, "FObjectTransmissionReceiveMessage must be trivial for the FProtocolQueuedItem::FMessage union.");
+	static_assert(std::is_trivial_v<FObjectSinkMessage>, "FObjectSinkMessage must be trivial for the FProtocolQueuedItem::FMessage union.");
+	
+	struct FProtocolQueuedItem
+	{
+		union FMessage
+		{
+			FInitMessage Init;
+			FObjectTraceBeginMessage ObjectTraceBegin;
+			FObjectTraceEndMessage ObjectTraceEnd;
+			FObjectTransmissionStartMessage TransmissionStart;
+			FObjectTransmissionReceiveMessage TransmissionReceive;
+			FObjectSinkMessage Sink;
+
+			FMessage() = default;
+			FMessage(const FInitMessage& InitMessage) : Init(InitMessage) {}
+			FMessage(const FObjectTraceBeginMessage& ObjectTraceBegin) : ObjectTraceBegin(ObjectTraceBegin) {}
+			FMessage(const FObjectTraceEndMessage& ObjectTraceEnd) : ObjectTraceEnd(ObjectTraceEnd) {}
+			FMessage(const FObjectTransmissionStartMessage& TransmissionStart) : TransmissionStart(TransmissionStart) {}
+			FMessage(const FObjectTransmissionReceiveMessage& TransmissionEnd) : TransmissionReceive(TransmissionEnd) {}
+			FMessage(const FObjectSinkMessage& Sink) : Sink(Sink) {}
+			
+		} Message;
+
+		/** Indicates which data in Message has been written to. */
+		EMessageType Type;
+
+		FProtocolQueuedItem() : Message(), Type(EMessageType::None) {}
+		template<typename TMessage>
+		FProtocolQueuedItem(const TMessage& Message) : Message(Message), Type(TMessage::Type()) {}
+	};
+	
+	/**
+	 * Used to synchronize analyzed data between two threads.
+	 * 
+	 * Aggregated .utrace files are analyzed by FProtocolEndpointAnalyzer on a separate thread.
+	 * FProtocolMultiEndpoint will load the data every tick on the main thread.
+	 */
+	class FProtocolDataQueue : public IProtocolDataTarget
+	{
+	public:
+
+		/** Queue of messages in the call order of the below functions. */
+		TQueue<FProtocolQueuedItem> MessageQueue;
+		
+		//~ Begin IProtocolDataTarget Interface
+		virtual void AppendInit(FInitMessage Message) override { MessageQueue.Enqueue(FProtocolQueuedItem(Message)); }
+		virtual void AppendObjectTraceBegin(FObjectTraceBeginMessage Message) override { MessageQueue.Enqueue(FProtocolQueuedItem(Message)); }
+		virtual void AppendObjectTraceEnd(FObjectTraceEndMessage Message) override { MessageQueue.Enqueue(FProtocolQueuedItem(Message)); }
+		virtual void AppendObjectTransmissionStart(FObjectTransmissionStartMessage Message) override { MessageQueue.Enqueue(FProtocolQueuedItem(Message)); }
+		virtual void AppendObjectTransmissionReceive(FObjectTransmissionReceiveMessage Message) override { MessageQueue.Enqueue(FProtocolQueuedItem(Message)); }
+		virtual void AppendObjectSink(FObjectSinkMessage Message) override { MessageQueue.Enqueue(FProtocolQueuedItem(Message)); }
+		//~ End IProtocolDataTarget Interface
+	};
+}
diff --git a/Engine/Plugins/Developer/Concert/ConcertInsights/ConcertInsightsVisualizer/Source/ConcertInsightsVisualizer/Private/Trace/Analysis/ProtocolEndpointAnalyzer.cpp b/Engine/Plugins/Developer/Concert/ConcertInsights/ConcertInsightsVisualizer/Source/ConcertInsightsVisualizer/Private/Trace/Analysis/ProtocolEndpointAnalyzer.cpp
new file mode 100644
--- /dev/null
+++ b/Engine/Plugins/Developer/Concert/ConcertInsights/ConcertInsightsVisualizer/Source/ConcertInsightsVisualizer/Private/Trace/Analysis/ProtocolEndpointAnalyzer.cpp
@@ -0,0 +1,77 @@
+// Copyright Epic Games, Inc. All Rights Reserved.
+
+#include "ProtocolEndpointAnalyzer.h"
+
+#include "IProtocolDataTarget.h"
+#include "LogConcertInsights.h"
+#include "Trace/Messages/InitMessage.h"
+#include "Trace/Messages/ObjectSinkMessage.h"
+#include "Trace/Messages/ObjectTraceMessage.h"
+#include "Trace/Messages/ObjectTransmissionReceiveMessage.h"
+#include "Trace/Messages/ObjectTransmissionStartMessage.h"
+
+namespace UE::ConcertInsightsVisualizer
+{
+	FProtocolEndpointAnalyzer::FProtocolEndpointAnalyzer(TraceServices::IAnalysisSession& Session, IProtocolDataTarget& DataTarget)
+		: Session(Session)
+		, DataTarget(DataTarget)
+	{}
+
+	void FProtocolEndpointAnalyzer::OnAnalysisBegin(const FOnAnalysisContext& Context)
+	{
+		FInterfaceBuilder& Builder = Context.InterfaceBuilder;
+		
+		Builder.RouteEvent(RouteId_Init, "ConcertLogger", "Init");
+		Builder.RouteEvent(RouteId_ObjectTraceBegin, "ConcertLogger", "ObjectTraceBegin");
+		Builder.RouteEvent(RouteId_ObjectTraceEnd, "ConcertLogger", "ObjectTraceEnd");
+		Builder.RouteEvent(RouteId_ObjectTransmissionStart, "ConcertLogger", "ObjectTransmissionStart");
+		Builder.RouteEvent(RouteId_ObjectTransmissionReceive, "ConcertLogger", "ObjectTransmissionReceive");
+		Builder.RouteEvent(RouteId_ObjectSink, "ConcertLogger", "ObjectSink");
+	}
+
+	bool FProtocolEndpointAnalyzer::OnEvent(uint16 RouteId, EStyle Style, const FOnEventContext& Context)
+	{
+		LLM_SCOPE_BYNAME(TEXT("Insights/FProtocolEndpointAnalyzer"));
+		
+		const FEventData& EventData = Context.EventData;
+		const FEventTime& EventTime = Context.EventTime;
+		switch (RouteId)
+		{
+		case RouteId_Init:
+		{
+			DataTarget.AppendInit(FInitMessage(EventData, EventTime, Session));
+			break;
+		}
+		case RouteId_ObjectTraceBegin:
+		{
+			DataTarget.AppendObjectTraceBegin(FObjectTraceBeginMessage(EventData, EventTime, Session));
+			break;
+		}
+		case RouteId_ObjectTraceEnd:
+		{
+			DataTarget.AppendObjectTraceEnd(FObjectTraceEndMessage(EventData, EventTime, Session));
+			break;
+		}
+		case RouteId_ObjectTransmissionStart:
+		{
+			DataTarget.AppendObjectTransmissionStart(FObjectTransmissionStartMessage(EventData, EventTime, Session));
+			break;
+		}
+		case RouteId_ObjectTransmissionReceive:
+		{
+			DataTarget.AppendObjectTransmissionReceive(FObjectTransmissionReceiveMessage(EventData, EventTime, Session));
+			break;
+		}
+		case RouteId_ObjectSink:
+		{
+			DataTarget.AppendObjectSink(FObjectSinkMessage(EventData, EventTime, Session));
+			break;
+		}
+			
+		default:
+		UE_LOG(LogConcertInsights, Warning, TEXT("Unknown RouteId %d"), RouteId);
+		}
+
+		return true;
+	}
+}
diff --git a/Engine/Plugins/Developer/Concert/ConcertInsights/ConcertInsightsVisualizer/Source/ConcertInsightsVisualizer/Private/Trace/Analysis/ProtocolEndpointAnalyzer.h b/Engine/Plugins/Developer/Concert/ConcertInsights/ConcertInsightsVisualizer/Source/ConcertInsightsVisualizer/Private/Trace/Analysis/ProtocolEndpointAnalyzer.h
new file mode 100644
--- /dev/null
+++ b/Engine/Plugins/Developer/Concert/ConcertInsights/ConcertInsightsVisualizer/Source/ConcertInsightsVisualizer/Private/Trace/Analysis/ProtocolEndpointAnalyzer.h
@@ -0,0 +1,49 @@
+// Copyright Epic Games, Inc. All Rights Reserved.
+
+#pragma once
+
+#include "Trace/Analyzer.h"
+
+namespace TraceServices { class IAnalysisSession; }
+
+namespace UE::ConcertInsightsVisualizer
+{
+	class IProtocolDataTarget;
+
+	/** Processes raw .utrace data, structures it into messages, and passes them to a IProtocolDataTarget. */
+	class FProtocolEndpointAnalyzer : public Trace::IAnalyzer
+	{
+	public:
+
+		/**
+		 * @param Session Caller is responsible for ensuring it outlives the constructed object.
+		 * @param DataTarget Caller is responsible for ensuring it outlives the constructed object.
+		 */
+		FProtocolEndpointAnalyzer(TraceServices::IAnalysisSession& Session UE_LIFETIMEBOUND, IProtocolDataTarget& DataTarget UE_LIFETIMEBOUND);
+
+		//~ Begin Trace::IAnalyzer Interface
+		virtual void OnAnalysisBegin(const FOnAnalysisContext& Context) override;
+		virtual void OnAnalysisEnd() override {}
+		virtual bool OnEvent(uint16 RouteId, EStyle Style, const FOnEventContext& Context) override;
+		//~ End Trace::IAnalyzer Interface
+		
+	private:
+		
+		enum : uint16
+		{
+			RouteId_Init,
+			RouteId_ObjectTraceBegin,
+			RouteId_ObjectTraceEnd,
+			RouteId_ObjectTransmissionStart,
+			RouteId_ObjectTransmissionReceive,
+			RouteId_ObjectSink
+		};
+
+		/** The session this provider operates on. */
+		TraceServices::IAnalysisSession& Session;
+
+		/** Stores the analyzed data. */
+		IProtocolDataTarget& DataTarget;
+	};
+}
+
diff --git a/Engine/Plugins/Developer/Concert/ConcertInsights/ConcertInsightsVisualizer/Source/ConcertInsightsVisualizer/Private/Trace/Messages/BaseObjectMessage.cpp b/Engine/Plugins/Developer/Concert/ConcertInsights/ConcertInsightsVisualizer/Source/ConcertInsightsVisualizer/Private/Trace/Messages/BaseObjectMessage.cpp
new file mode 100644
--- /dev/null
+++ b/Engine/Plugins/Developer/Concert/ConcertInsights/ConcertInsightsVisualizer/Source/ConcertInsightsVisualizer/Private/Trace/Messages/BaseObjectMessage.cpp
@@ -0,0 +1,30 @@
+// Copyright Epic Games, Inc. All Rights Reserved.
+
+#include "BaseObjectMessage.h"
+
+#include "TraceServices/Model/AnalysisSession.h"
+
+namespace UE::ConcertInsightsVisualizer
+{
+	namespace ObjectTraceMessage
+	{
+		static const TCHAR* GetObjectPath(
+			const Trace::IAnalyzer::FEventData& EventData,
+			TraceServices::IAnalysisSession& Session UE_LIFETIMEBOUND
+			)
+		{
+			FString ObjectPath;
+			EventData.GetString("ObjectPath", ObjectPath);
+			return Session.StoreString(*ObjectPath);
+		}
+	}
+	
+	FBaseObjectMessage::FBaseObjectMessage(
+		const Trace::IAnalyzer::FEventData& EventData,
+		TraceServices::IAnalysisSession& Session UE_LIFETIMEBOUND
+		)
+		: Protocol(EventData.GetValue<FProtocolId>("Protocol"))
+		, SequenceId(EventData.GetValue<FSequenceId>("SequenceId"))
+		, ObjectPath(ObjectTraceMessage::GetObjectPath(EventData, Session))
+	{}
+}
diff --git a/Engine/Plugins/Developer/Concert/ConcertInsights/ConcertInsightsVisualizer/Source/ConcertInsightsVisualizer/Private/Trace/Messages/BaseObjectMessage.h b/Engine/Plugins/Developer/Concert/ConcertInsights/ConcertInsightsVisualizer/Source/ConcertInsightsVisualizer/Private/Trace/Messages/BaseObjectMessage.h
new file mode 100644
--- /dev/null
+++ b/Engine/Plugins/Developer/Concert/ConcertInsights/ConcertInsightsVisualizer/Source/ConcertInsightsVisualizer/Private/Trace/Messages/BaseObjectMessage.h
@@ -0,0 +1,32 @@
+// Copyright Epic Games, Inc. All Rights Reserved.
+
+#pragma once
+
+#include "Trace/Types/ProtocolId.h"
+#include "Trace/Types/SequenceId.h"
+
+#include "HAL/Platform.h"
+#include "Trace/Analyzer.h"
+#include "UObject/SoftObjectPath.h"
+
+#include <type_traits>
+
+namespace TraceServices { class IAnalysisSession; }
+
+namespace UE::ConcertInsightsVisualizer
+{
+	/** Shared data for messages about objects. */
+	struct FBaseObjectMessage
+	{
+		FProtocolId Protocol;
+		FSequenceId SequenceId;
+		/** Points to string stored in IAnalysisSession::StoreString */
+		const TCHAR* ObjectPath;
+		
+		FBaseObjectMessage() = default;
+		FBaseObjectMessage(const Trace::IAnalyzer::FEventData& EventData, TraceServices::IAnalysisSession& Session UE_LIFETIMEBOUND);
+
+		FSoftObjectPath GetSoftObjectPath() const { return FSoftObjectPath(ObjectPath); }
+	};
+	static_assert(std::is_trivial_v<FBaseObjectMessage>, "FBaseObjectMessage must be trivial for the FProtocolQueuedItem::FMessage union.");
+}
diff --git a/Engine/Plugins/Developer/Concert/ConcertInsights/ConcertInsightsVisualizer/Source/ConcertInsightsVisualizer/Private/Trace/Messages/EMessageType.h b/Engine/Plugins/Developer/Concert/ConcertInsights/ConcertInsightsVisualizer/Source/ConcertInsightsVisualizer/Private/Trace/Messages/EMessageType.h
new file mode 100644
--- /dev/null
+++ b/Engine/Plugins/Developer/Concert/ConcertInsights/ConcertInsightsVisualizer/Source/ConcertInsightsVisualizer/Private/Trace/Messages/EMessageType.h
@@ -0,0 +1,20 @@
+// Copyright Epic Games, Inc. All Rights Reserved.
+
+#pragma once
+
+#include "HAL/Platform.h"
+
+namespace UE::ConcertInsightsVisualizer
+{
+	enum class EMessageType : uint8
+	{
+		None,
+		Init,
+		ObjectTraceBegin,
+		ObjectTraceEnd,
+		ObjectTransmissionStart,
+		ObjectTransmissionReceive,
+		ObjectSink
+	};
+}
+
diff --git a/Engine/Plugins/Developer/Concert/ConcertInsights/ConcertInsightsVisualizer/Source/ConcertInsightsVisualizer/Private/Trace/Messages/InitMessage.cpp b/Engine/Plugins/Developer/Concert/ConcertInsights/ConcertInsightsVisualizer/Source/ConcertInsightsVisualizer/Private/Trace/Messages/InitMessage.cpp
new file mode 100644
--- /dev/null
+++ b/Engine/Plugins/Developer/Concert/ConcertInsights/ConcertInsightsVisualizer/Source/ConcertInsightsVisualizer/Private/Trace/Messages/InitMessage.cpp
@@ -0,0 +1,35 @@
+// Copyright Epic Games, Inc. All Rights Reserved.
+
+#include "InitMessage.h"
+
+#include "TraceServices/Model/AnalysisSession.h"
+
+namespace UE::ConcertInsightsVisualizer
+{
+	namespace InitMessage
+	{
+		static const TCHAR* GetClientDisplayName(
+			const Trace::IAnalyzer::FEventData& EventData,
+			TraceServices::IAnalysisSession& Session UE_LIFETIMEBOUND
+			)
+		{
+			const bool bHasDisplayName = EventData.GetValue<bool>("HasDisplayName");
+			FString ClientDisplayName;
+			EventData.GetString("ClientDisplayName", ClientDisplayName);
+			return bHasDisplayName ? Session.StoreString(ClientDisplayName) : nullptr;
+		}
+	}
+
+	FInitMessage::FInitMessage(const Trace::IAnalyzer::FEventData& EventData, const Trace::IAnalyzer::FEventTime& EventTime, TraceServices::IAnalysisSession& Session)
+		: Version(static_cast<ConcertTrace::EConcertTraceVersion>(EventData.GetValue<uint8>("Version")))
+		, bHasEndpointId(EventData.GetValue<bool>("HasEndpointId"))
+		, EndpointId_A(EventData.GetValue<uint32>("EndpointId_A"))
+		, EndpointId_B(EventData.GetValue<uint32>("EndpointId_B"))
+		, EndpointId_C(EventData.GetValue<uint32>("EndpointId_C"))
+		, EndpointId_D(EventData.GetValue<uint32>("EndpointId_D"))
+		, TraceInitUtc_Ticks(EventData.GetValue<int64>("DateTimeTicks"))
+		, StartTime(EventTime.AsSeconds(EventData.GetValue<uint64>("Cycle")))
+		, ClientDisplayName(InitMessage::GetClientDisplayName(EventData, Session))
+		, bIsServer(EventData.GetValue<bool>("IsServer"))
+	{}
+}
diff --git a/Engine/Plugins/Developer/Concert/ConcertInsights/ConcertInsightsVisualizer/Source/ConcertInsightsVisualizer/Private/Trace/Messages/InitMessage.h b/Engine/Plugins/Developer/Concert/ConcertInsights/ConcertInsightsVisualizer/Source/ConcertInsightsVisualizer/Private/Trace/Messages/InitMessage.h
new file mode 100644
--- /dev/null
+++ b/Engine/Plugins/Developer/Concert/ConcertInsights/ConcertInsightsVisualizer/Source/ConcertInsightsVisualizer/Private/Trace/Messages/InitMessage.h
@@ -0,0 +1,58 @@
+// Copyright Epic Games, Inc. All Rights Reserved.
+
+#pragma once
+
+#include "EMessageType.h"
+
+#include "HAL/Platform.h"
+#include "Misc/DateTime.h"
+#include "Misc/Guid.h"
+#include "Misc/Optional.h"
+#include "Trace/Analyzer.h"
+
+#include <type_traits>
+
+namespace UE::ConcertTrace { enum class EConcertTraceVersion : uint8; }
+namespace TraceServices { class IAnalysisSession; }
+
+namespace UE::ConcertInsightsVisualizer
+{
+	struct FInitMessage
+	{
+		FInitMessage() = default;
+		FInitMessage(const Trace::IAnalyzer::FEventData& EventData, const Trace::IAnalyzer::FEventTime& EventTime, TraceServices::IAnalysisSession& Session UE_LIFETIMEBOUND);
+		
+		static constexpr EMessageType Type() { return EMessageType::Init; }
+
+		ConcertTrace::EConcertTraceVersion GetVersion() const { return Version; }
+		TOptional<FGuid> GetEndpointId() const { return bHasEndpointId ? FGuid(EndpointId_A, EndpointId_B, EndpointId_C, EndpointId_D) : TOptional<FGuid>(); }
+		/** The UTC time reported on the sending machine when the init event was started. Used to correlate cycles across multiple machines. */
+		FDateTime GetTraceInitTimeUtc() const { return FDateTime(TraceInitUtc_Ticks); }
+		/** Seconds since start of the trace. */
+		double GetStartTime() const { return StartTime; }
+		/** Points to string stored in IAnalysisSession::StoreString */
+		TOptional<const TCHAR*> GetClientDisplayName() const { return ClientDisplayName; }
+		bool IsServer() const { return bIsServer; }
+
+	private: // This struct is a bit more complicated so let's force use of the above getters
+		
+		ConcertTrace::EConcertTraceVersion Version;
+
+		bool bHasEndpointId;
+		uint32 EndpointId_A;
+		uint32 EndpointId_B;
+		uint32 EndpointId_C;
+		uint32 EndpointId_D;
+		
+		/** The UTC time reported on the sending machine when the init event was started. Used to correlate cycles across multiple machines. */
+		int64 TraceInitUtc_Ticks;
+		/** Seconds since start of the trace. */
+		double StartTime;
+		
+		/** Points to string stored in IAnalysisSession::StoreString */
+		const TCHAR* ClientDisplayName;
+		bool bIsServer;
+	};
+
+	static_assert(std::is_trivial_v<FInitMessage>, "FInitMessage must be unique so it can be put into a union.");
+}
diff --git a/Engine/Plugins/Developer/Concert/ConcertInsights/ConcertInsightsVisualizer/Source/ConcertInsightsVisualizer/Private/Trace/Messages/ObjectSinkMessage.cpp b/Engine/Plugins/Developer/Concert/ConcertInsights/ConcertInsightsVisualizer/Source/ConcertInsightsVisualizer/Private/Trace/Messages/ObjectSinkMessage.cpp
new file mode 100644
--- /dev/null
+++ b/Engine/Plugins/Developer/Concert/ConcertInsights/ConcertInsightsVisualizer/Source/ConcertInsightsVisualizer/Private/Trace/Messages/ObjectSinkMessage.cpp
@@ -0,0 +1,31 @@
+// Copyright Epic Games, Inc. All Rights Reserved.
+
+#include "ObjectSinkMessage.h"
+
+#include "TraceServices/Model/AnalysisSession.h"
+
+namespace UE::ConcertInsightsVisualizer
+{
+	namespace ObjectSinkMessage
+	{
+		static const TCHAR* GetSinkName(
+			const Trace::IAnalyzer::FEventData& EventData,
+			TraceServices::IAnalysisSession& Session UE_LIFETIMEBOUND
+			)
+		{
+			FString SinkName;
+			EventData.GetString("SinkName", SinkName);
+			return Session.StoreString(*SinkName);
+		}
+	}
+	
+	FObjectSinkMessage::FObjectSinkMessage(
+		const Trace::IAnalyzer::FEventData& EventData,
+		const Trace::IAnalyzer::FEventTime& EventTime,
+		TraceServices::IAnalysisSession& Session UE_LIFETIMEBOUND
+		)
+		: FBaseObjectMessage(EventData, Session)
+		, Time(EventTime.AsSeconds(EventData.GetValue<uint64>("Cycle")))
+		, SinkName(ObjectSinkMessage::GetSinkName(EventData, Session))
+	{}
+}
diff --git a/Engine/Plugins/Developer/Concert/ConcertInsights/ConcertInsightsVisualizer/Source/ConcertInsightsVisualizer/Private/Trace/Messages/ObjectSinkMessage.h b/Engine/Plugins/Developer/Concert/ConcertInsights/ConcertInsightsVisualizer/Source/ConcertInsightsVisualizer/Private/Trace/Messages/ObjectSinkMessage.h
new file mode 100644
--- /dev/null
+++ b/Engine/Plugins/Developer/Concert/ConcertInsights/ConcertInsightsVisualizer/Source/ConcertInsightsVisualizer/Private/Trace/Messages/ObjectSinkMessage.h
@@ -0,0 +1,34 @@
+// Copyright Epic Games, Inc. All Rights Reserved.
+
+#pragma once
+
+#include "BaseObjectMessage.h"
+#include "EMessageType.h"
+
+#include "HAL/Platform.h"
+#include "Trace/Analyzer.h"
+
+#include <type_traits>
+
+namespace UE::ConcertInsightsVisualizer
+{
+	struct FObjectSinkMessage : FBaseObjectMessage
+	{
+		/** The time at which the message was generated. Seconds since start of trace. */
+		double Time;
+		
+		/** Points to string stored in IAnalysisSession::StoreString */
+		const TCHAR* SinkName;
+		
+		FObjectSinkMessage() = default;
+		FObjectSinkMessage(
+			const Trace::IAnalyzer::FEventData& EventData,
+			const Trace::IAnalyzer::FEventTime& EventTime,
+			TraceServices::IAnalysisSession& Session UE_LIFETIMEBOUND
+			);
+		
+		static constexpr EMessageType Type() { return EMessageType::ObjectSink; }
+	};
+	
+	static_assert(std::is_trivial_v<FObjectSinkMessage>, "FObjectSinkMessage must be unique so it can be put into a union.");
+}
diff --git a/Engine/Plugins/Developer/Concert/ConcertInsights/ConcertInsightsVisualizer/Source/ConcertInsightsVisualizer/Private/Trace/Messages/ObjectTraceMessage.cpp b/Engine/Plugins/Developer/Concert/ConcertInsights/ConcertInsightsVisualizer/Source/ConcertInsightsVisualizer/Private/Trace/Messages/ObjectTraceMessage.cpp
new file mode 100644
--- /dev/null
+++ b/Engine/Plugins/Developer/Concert/ConcertInsights/ConcertInsightsVisualizer/Source/ConcertInsightsVisualizer/Private/Trace/Messages/ObjectTraceMessage.cpp
@@ -0,0 +1,31 @@
+// Copyright Epic Games, Inc. All Rights Reserved.
+
+#include "ObjectTraceMessage.h"
+
+#include "TraceServices/Model/AnalysisSession.h"
+
+namespace UE::ConcertInsightsVisualizer
+{
+	namespace ObjectTraceMessage
+	{
+		static const TCHAR* GetEventName(
+			const Trace::IAnalyzer::FEventData& EventData,
+			TraceServices::IAnalysisSession& Session UE_LIFETIMEBOUND
+			)
+		{
+			FString EventName;
+			EventData.GetString("EventName", EventName);
+			return Session.StoreString(*EventName);
+		}
+	}
+	
+	FObjectTraceMessage::FObjectTraceMessage(
+		const Trace::IAnalyzer::FEventData& EventData,
+		const Trace::IAnalyzer::FEventTime& EventTime,
+		TraceServices::IAnalysisSession& Session UE_LIFETIMEBOUND
+		)
+		: FBaseObjectMessage(EventData, Session)
+		, Time(EventTime.AsSeconds(EventData.GetValue<uint64>("Cycle")))
+		, EventName(ObjectTraceMessage::GetEventName(EventData, Session))
+	{}
+}
diff --git a/Engine/Plugins/Developer/Concert/ConcertInsights/ConcertInsightsVisualizer/Source/ConcertInsightsVisualizer/Private/Trace/Messages/ObjectTraceMessage.h b/Engine/Plugins/Developer/Concert/ConcertInsights/ConcertInsightsVisualizer/Source/ConcertInsightsVisualizer/Private/Trace/Messages/ObjectTraceMessage.h
new file mode 100644
--- /dev/null
+++ b/Engine/Plugins/Developer/Concert/ConcertInsights/ConcertInsightsVisualizer/Source/ConcertInsightsVisualizer/Private/Trace/Messages/ObjectTraceMessage.h
@@ -0,0 +1,55 @@
+// Copyright Epic Games, Inc. All Rights Reserved.
+
+#pragma once
+
+#include "BaseObjectMessage.h"
+#include "EMessageType.h"
+
+#include "HAL/Platform.h"
+#include "Trace/Analyzer.h"
+
+namespace UE::ConcertInsightsVisualizer
+{
+	struct FObjectTraceMessage : FBaseObjectMessage
+	{
+		/** Time of the event. Seconds since start of trace. */
+		double Time;
+		const TCHAR* EventName;
+		
+		FObjectTraceMessage() = default;
+		FObjectTraceMessage(
+			const Trace::IAnalyzer::FEventData& EventData,
+			const Trace::IAnalyzer::FEventTime& EventTime,
+			TraceServices::IAnalysisSession& Session UE_LIFETIMEBOUND
+			);
+	};
+	static_assert(std::is_trivial_v<FObjectTraceMessage>, "FObjectTraceMessage must be unique so it can be put into a union.");
+
+	struct FObjectTraceBeginMessage : FObjectTraceMessage
+	{
+		FObjectTraceBeginMessage() = default;
+		FObjectTraceBeginMessage(
+			const Trace::IAnalyzer::FEventData& EventData,
+			const Trace::IAnalyzer::FEventTime& EventTime,
+			TraceServices::IAnalysisSession& Session UE_LIFETIMEBOUND
+			) : FObjectTraceMessage(EventData, EventTime, Session)
+		{}
+		
+		static constexpr EMessageType Type() { return EMessageType::ObjectTraceBegin; }
+	};
+	static_assert(std::is_trivial_v<FObjectTraceBeginMessage>, "FObjectTraceBeginMessage must be trivial for the FProtocolQueuedItem::FMessage union.");
+	
+	struct FObjectTraceEndMessage : FObjectTraceMessage
+	{
+		FObjectTraceEndMessage() = default;
+		FObjectTraceEndMessage(
+			const Trace::IAnalyzer::FEventData& EventData,
+			const Trace::IAnalyzer::FEventTime& EventTime,
+			TraceServices::IAnalysisSession& Session UE_LIFETIMEBOUND
+			) : FObjectTraceMessage(EventData, EventTime, Session)
+		{}
+		
+		static constexpr EMessageType Type() { return EMessageType::ObjectTraceEnd; }
+	};
+	static_assert(std::is_trivial_v<FObjectTraceEndMessage>, "FObjectTraceEndMessage must be unique so it can be put into a union.");
+}
diff --git a/Engine/Plugins/Developer/Concert/ConcertInsights/ConcertInsightsVisualizer/Source/ConcertInsightsVisualizer/Private/Trace/Messages/ObjectTransmissionReceiveMessage.cpp b/Engine/Plugins/Developer/Concert/ConcertInsights/ConcertInsightsVisualizer/Source/ConcertInsightsVisualizer/Private/Trace/Messages/ObjectTransmissionReceiveMessage.cpp
new file mode 100644
--- /dev/null
+++ b/Engine/Plugins/Developer/Concert/ConcertInsights/ConcertInsightsVisualizer/Source/ConcertInsightsVisualizer/Private/Trace/Messages/ObjectTransmissionReceiveMessage.cpp
@@ -0,0 +1,15 @@
+// Copyright Epic Games, Inc. All Rights Reserved.
+
+#include "ObjectTransmissionReceiveMessage.h"
+
+namespace UE::ConcertInsightsVisualizer
+{
+	FObjectTransmissionReceiveMessage::FObjectTransmissionReceiveMessage(
+		const Trace::IAnalyzer::FEventData& EventData,
+		const Trace::IAnalyzer::FEventTime& EventTime,
+		TraceServices::IAnalysisSession& Session UE_LIFETIMEBOUND
+		)
+		: FBaseObjectMessage(EventData, Session)
+		, Time(EventTime.AsSeconds(EventData.GetValue<uint64>("Cycle")))
+	{}
+}
diff --git a/Engine/Plugins/Developer/Concert/ConcertInsights/ConcertInsightsVisualizer/Source/ConcertInsightsVisualizer/Private/Trace/Messages/ObjectTransmissionReceiveMessage.h b/Engine/Plugins/Developer/Concert/ConcertInsights/ConcertInsightsVisualizer/Source/ConcertInsightsVisualizer/Private/Trace/Messages/ObjectTransmissionReceiveMessage.h
new file mode 100644
--- /dev/null
+++ b/Engine/Plugins/Developer/Concert/ConcertInsights/ConcertInsightsVisualizer/Source/ConcertInsightsVisualizer/Private/Trace/Messages/ObjectTransmissionReceiveMessage.h
@@ -0,0 +1,30 @@
+// Copyright Epic Games, Inc. All Rights Reserved.
+
+#pragma once
+
+#include "BaseObjectMessage.h"
+#include "EMessageType.h"
+
+#include "HAL/Platform.h"
+#include "Trace/Analyzer.h"
+
+#include <type_traits>
+
+namespace UE::ConcertInsightsVisualizer
+{
+	struct FObjectTransmissionReceiveMessage : FBaseObjectMessage
+	{
+		/** The time at which the message was generated. Seconds since start of trace. */
+		double Time;
+		
+		FObjectTransmissionReceiveMessage() = default;
+		FObjectTransmissionReceiveMessage(
+			const Trace::IAnalyzer::FEventData& EventData,
+			const Trace::IAnalyzer::FEventTime& EventTime,
+			TraceServices::IAnalysisSession& Session UE_LIFETIMEBOUND
+			);
+		
+		static constexpr EMessageType Type() { return EMessageType::ObjectTransmissionReceive; }
+	};
+	static_assert(std::is_trivial_v<FObjectTransmissionReceiveMessage>, "FObjectTransmissionReceiveMessage must be unique so it can be put into a union.");
+}
diff --git a/Engine/Plugins/Developer/Concert/ConcertInsights/ConcertInsightsVisualizer/Source/ConcertInsightsVisualizer/Private/Trace/Messages/ObjectTransmissionStartMessage.cpp b/Engine/Plugins/Developer/Concert/ConcertInsights/ConcertInsightsVisualizer/Source/ConcertInsightsVisualizer/Private/Trace/Messages/ObjectTransmissionStartMessage.cpp
new file mode 100644
--- /dev/null
+++ b/Engine/Plugins/Developer/Concert/ConcertInsights/ConcertInsightsVisualizer/Source/ConcertInsightsVisualizer/Private/Trace/Messages/ObjectTransmissionStartMessage.cpp
@@ -0,0 +1,31 @@
+// Copyright Epic Games, Inc. All Rights Reserved.
+
+#include "ObjectTransmissionStartMessage.h"
+
+namespace UE::ConcertInsightsVisualizer
+{
+	namespace ObjectTransmissionStartMessage
+	{
+		static FGuid GetDestinationEndpointId(const Trace::IAnalyzer::FEventData& EventData)
+		{
+			const uint32 EndpointId_A = EventData.GetValue<uint32>("DestEndpointId_A");
+			const uint32 EndpointId_B = EventData.GetValue<uint32>("DestEndpointId_B");
+			const uint32 EndpointId_C = EventData.GetValue<uint32>("DestEndpointId_C");
+			const uint32 EndpointId_D = EventData.GetValue<uint32>("DestEndpointId_D");
+			return FGuid(EndpointId_A, EndpointId_B, EndpointId_C, EndpointId_D);
+		}
+	}
+	
+	FObjectTransmissionStartMessage::FObjectTransmissionStartMessage(
+		const Trace::IAnalyzer::FEventData& EventData,
+		const Trace::IAnalyzer::FEventTime& EventTime,
+		TraceServices::IAnalysisSession& Session UE_LIFETIMEBOUND
+		)
+		: FBaseObjectMessage(EventData, Session)
+		, DestinationEndpointId_A(EventData.GetValue<uint32>("DestEndpointId_A"))
+		, DestinationEndpointId_B(EventData.GetValue<uint32>("DestEndpointId_B"))
+		, DestinationEndpointId_C(EventData.GetValue<uint32>("DestEndpointId_C"))
+		, DestinationEndpointId_D(EventData.GetValue<uint32>("DestEndpointId_D"))
+		, Time(EventTime.AsSeconds(EventData.GetValue<uint64>("Cycle")))
+	{}
+}
diff --git a/Engine/Plugins/Developer/Concert/ConcertInsights/ConcertInsightsVisualizer/Source/ConcertInsightsVisualizer/Private/Trace/Messages/ObjectTransmissionStartMessage.h b/Engine/Plugins/Developer/Concert/ConcertInsights/ConcertInsightsVisualizer/Source/ConcertInsightsVisualizer/Private/Trace/Messages/ObjectTransmissionStartMessage.h
new file mode 100644
--- /dev/null
+++ b/Engine/Plugins/Developer/Concert/ConcertInsights/ConcertInsightsVisualizer/Source/ConcertInsightsVisualizer/Private/Trace/Messages/ObjectTransmissionStartMessage.h
@@ -0,0 +1,41 @@
+// Copyright Epic Games, Inc. All Rights Reserved.
+
+#pragma once
+
+#include "BaseObjectMessage.h"
+#include "EMessageType.h"
+
+#include "HAL/Platform.h"
+#include "Trace/Analyzer.h"
+
+#include <type_traits>
+
+namespace UE::ConcertInsightsVisualizer
+{
+	struct FObjectTransmissionStartMessage : FBaseObjectMessage
+	{
+		FObjectTransmissionStartMessage() = default;
+		FObjectTransmissionStartMessage(
+			const Trace::IAnalyzer::FEventData& EventData,
+			const Trace::IAnalyzer::FEventTime& EventTime,
+			TraceServices::IAnalysisSession& Session UE_LIFETIMEBOUND
+			);
+		
+		static constexpr EMessageType Type() { return EMessageType::ObjectTransmissionStart; }
+
+		FGuid GetDestinationEdnpointId() const { return FGuid(DestinationEndpointId_A, DestinationEndpointId_B, DestinationEndpointId_C, DestinationEndpointId_D); }
+		double GetTime() const { return Time; }
+
+	private: // This struct is a bit more complicated so let's force use of the above getters
+
+		// Used to form FGuid for the destination endpoint
+		int64 DestinationEndpointId_A;
+		int64 DestinationEndpointId_B;
+		int64 DestinationEndpointId_C;
+		int64 DestinationEndpointId_D;
+		
+		/** The time at which the message was generated. Seconds since start of trace. */
+		double Time;
+	};
+	static_assert(std::is_trivial_v<FObjectTransmissionStartMessage>, "FObjectTransmissionStartMessage must be unique so it can be put into a union.");
+}
diff --git a/Engine/Plugins/Developer/Concert/ConcertInsights/ConcertInsightsVisualizer/Source/ConcertInsightsVisualizer/Private/Trace/ObjectSequence/ObjectNetworkScope.h b/Engine/Plugins/Developer/Concert/ConcertInsights/ConcertInsightsVisualizer/Source/ConcertInsightsVisualizer/Private/Trace/ObjectSequence/ObjectNetworkScope.h
new file mode 100644
--- /dev/null
+++ b/Engine/Plugins/Developer/Concert/ConcertInsights/ConcertInsightsVisualizer/Source/ConcertInsightsVisualizer/Private/Trace/ObjectSequence/ObjectNetworkScope.h
@@ -0,0 +1,26 @@
+// Copyright Epic Games, Inc. All Rights Reserved.
+
+#pragma once
+
+#include "Trace/Types/EndpointId.h"
+
+#include "HAL/Platform.h"
+#include "Misc/Optional.h"
+
+namespace UE::ConcertInsightsVisualizer
+{
+	
+	/** Nested right under FObjectSequence. Indicates on which client the object was (or whether it was in transit). */
+	struct FObjectNetworkScope
+	{
+		/**
+		 * The endpoint that corresponds to this scope, if any.
+		 * Unset means that this is a network transit of the packet from one to another endpoint.
+		 * FEndpointId can be used to resolve the display name in FProtocolMultiEndpointProvider.
+		 */
+		TOptional<FEndpointId> ProcessingEndpoint;
+
+		bool IsEndpoint() const { return ProcessingEndpoint.IsSet(); }
+		bool IsTransportAcrossNetwork() const { return !IsEndpoint(); }
+	};
+}
diff --git a/Engine/Plugins/Developer/Concert/ConcertInsights/ConcertInsightsVisualizer/Source/ConcertInsightsVisualizer/Private/Trace/ObjectSequence/ObjectProcessingStep.h b/Engine/Plugins/Developer/Concert/ConcertInsights/ConcertInsightsVisualizer/Source/ConcertInsightsVisualizer/Private/Trace/ObjectSequence/ObjectProcessingStep.h
new file mode 100644
--- /dev/null
+++ b/Engine/Plugins/Developer/Concert/ConcertInsights/ConcertInsightsVisualizer/Source/ConcertInsightsVisualizer/Private/Trace/ObjectSequence/ObjectProcessingStep.h
@@ -0,0 +1,15 @@
+// Copyright Epic Games, Inc. All Rights Reserved.
+
+#pragma once
+
+#include "HAL/Platform.h"
+
+namespace UE::ConcertInsightsVisualizer
+{
+	/** Nested under FObjectNetworkScope. Describes what an endpoint spent CPU time on while processing an object in a sequence. */
+	struct FObjectProcessingStep
+	{
+		/** Name of the processing step */
+		const TCHAR* EventName;
+	};
+}
\ No newline at end of file
diff --git a/Engine/Plugins/Developer/Concert/ConcertInsights/ConcertInsightsVisualizer/Source/ConcertInsightsVisualizer/Private/Trace/ObjectSequence/ObjectSequence.h b/Engine/Plugins/Developer/Concert/ConcertInsights/ConcertInsightsVisualizer/Source/ConcertInsightsVisualizer/Private/Trace/ObjectSequence/ObjectSequence.h
new file mode 100644
--- /dev/null
+++ b/Engine/Plugins/Developer/Concert/ConcertInsights/ConcertInsightsVisualizer/Source/ConcertInsightsVisualizer/Private/Trace/ObjectSequence/ObjectSequence.h
@@ -0,0 +1,18 @@
+// Copyright Epic Games, Inc. All Rights Reserved.
+
+#pragma once
+
+#include "Trace/Types/SequenceId.h"
+
+namespace UE::ConcertInsightsVisualizer
+{
+	/**
+	 * A sequence is a time range in which an object was processed across multiple endpoints.
+	 * This data is displayed top-most event.
+	 */
+	struct FObjectSequence
+	{
+		/** ID for the object update. */
+		FSequenceId SequenceId;
+	};
+}
diff --git a/Engine/Plugins/Developer/Concert/ConcertInsights/ConcertInsightsVisualizer/Source/ConcertInsightsVisualizer/Private/Trace/ProtocolMultiEndpointProvider.cpp b/Engine/Plugins/Developer/Concert/ConcertInsights/ConcertInsightsVisualizer/Source/ConcertInsightsVisualizer/Private/Trace/ProtocolMultiEndpointProvider.cpp
new file mode 100644
--- /dev/null
+++ b/Engine/Plugins/Developer/Concert/ConcertInsights/ConcertInsightsVisualizer/Source/ConcertInsightsVisualizer/Private/Trace/ProtocolMultiEndpointProvider.cpp
@@ -0,0 +1,536 @@
+// Copyright Epic Games, Inc. All Rights Reserved.
+
+#include "ProtocolMultiEndpointProvider.h"
+
+#include "LogConcertInsights.h"
+
+#include "Algo/AnyOf.h"
+#include "Insights/IUnrealInsightsModule.h"
+#include "Modules/ModuleManager.h"
+#include "Util/TimeSyncUtils.h"
+
+namespace UE::ConcertInsightsVisualizer
+{
+	namespace Private
+	{
+		/** Given a time window of Insights, should the given scope be displayed? */
+		static bool ShouldDisplayScopeInWindow(double WindowStart, double WindowEnd, double ScopeStart, double ScopeEnd)
+		{
+			const bool bDoesWindowIncludeStart = WindowStart <= ScopeStart && ScopeStart <= WindowEnd;
+			const bool bDoesWindowIncludeEnd = WindowStart <= ScopeEnd && ScopeEnd <= WindowEnd;
+			const bool bDoesSequenceIncludeWindow = ScopeStart <= WindowStart && WindowEnd <= ScopeEnd;
+			
+			const bool bShouldDisplayScope = bDoesWindowIncludeStart || bDoesWindowIncludeEnd || bDoesSequenceIncludeWindow;
+			return bShouldDisplayScope;
+		}
+
+		static void LogInitMessage(const TCHAR* Preamble, FEndpointId TraceFileId, const FInitMessage& Init)
+		{
+			UE_LOG(LogConcertInsights, Log, TEXT("%s %llu: bIsServer: %d, DisplayName: %s, UTC: %s"),
+				Preamble,
+				TraceFileId,
+				Init.IsServer(),
+				Init.GetClientDisplayName() ? *Init.GetClientDisplayName() : TEXT("none"),
+				*Init.GetTraceInitTimeUtc().ToString(TEXT("%Y-%m-%d_%H-%Mm-%Ss-%sms"))
+				);
+		}
+	}
+	
+	const FName FProtocolMultiEndpointProvider::ProviderName = TEXT("FProtocolMultiEndpointProvider");
+	
+	FProtocolMultiEndpointProvider::FProtocolMultiEndpointProvider(TraceServices::IAnalysisSession& Session UE_LIFETIMEBOUND)
+		: Session(Session)
+		, Aggregator(*FModuleManager::LoadModuleChecked<IUnrealInsightsModule>("TraceInsights").GetStoreClient(), GetMainTraceId())
+		, TickHandle(FTSTicker::GetCoreTicker().AddTicker(FTickerDelegate::CreateRaw(this, &FProtocolMultiEndpointProvider::ProcessGeneratedTraceData)))
+	{}
+
+	FProtocolMultiEndpointProvider::~FProtocolMultiEndpointProvider()
+	{
+		FTSTicker::GetCoreTicker().RemoveTicker(TickHandle);
+	}
+
+	const TCHAR* FProtocolMultiEndpointProvider::GetEndpointDisplayName(FEndpointId EndpointId) const
+	{
+		Session.ReadAccessCheck();
+		
+		const bool bIsMainTrace = EndpointId == Session.GetTraceId();
+		const FInitMessage* EndpointInitdata = bIsMainTrace
+			? MainTraceInitData ? &MainTraceInitData.GetValue() : nullptr
+			: EndpointInitData.Find(EndpointId);
+		
+		const TCHAR* EmptyString = TEXT("Null");
+		return EndpointInitdata
+			? EndpointInitdata->GetClientDisplayName().Get(EmptyString)
+			: EmptyString;
+	}
+
+	void FProtocolMultiEndpointProvider::EnumerateProtocols(TFunctionRef<TraceServices::EEventEnumerate(FProtocolId)> Callback) const
+	{
+		Session.ReadAccessCheck();
+
+		for (FProtocolId ProtocolId : ActiveProtocols)
+		{
+			if (Callback(ProtocolId) == TraceServices::EEventEnumerate::Stop)
+			{
+				break;
+			}
+		}
+	}
+
+	void FProtocolMultiEndpointProvider::EnumerateObjects(FProtocolId Protocol, TFunctionRef<TraceServices::EEventEnumerate(FObjectPath Object)> Callback) const
+	{
+		Session.ReadAccessCheck();
+
+		for (const TPair<FObjectScopeInfo, FPerObjectData>& ObjectPair : PerObjectData)
+		{
+			if (ObjectPair.Key.ProtocolId == Protocol
+				&& Callback(ObjectPair.Key.ObjectPath) == TraceServices::EEventEnumerate::Stop)
+			{
+				break;
+			}
+		}
+	}
+
+	void FProtocolMultiEndpointProvider::EnumerateEndpointsInSequence(const FSequenceScopeInfo& Info, TFunctionRef<TraceServices::EEventEnumerate(FEndpointId EndpointId)> Callback) const
+	{
+		Session.ReadAccessCheck();
+
+		const FPerSequenceData* SequenceData = FindSequenceData(Info);
+		if (!SequenceData)
+		{
+			return;
+		}
+
+		for (FEndpointId EndpointId : SequenceData->Endpoints)
+		{
+			if (Callback(EndpointId) == TraceServices::EEventEnumerate::Stop)
+			{
+				break;
+			}
+		}
+	}
+
+	void FProtocolMultiEndpointProvider::EnumerateSequences(double Start, double End, const FObjectScopeInfo& Info, TFunctionRef<TraceServices::EEventEnumerate(FSequenceId Sequence)> Callback) const
+	{
+		Session.ReadAccessCheck();
+
+		const FPerObjectData* ObjectData = FindObjectData(Info);
+		if (!ObjectData)
+		{
+			return;
+		}
+
+		// CachedTimelineSortedByStart or one of the sequence times may have changed since the last call. If so, resort.
+		if (ObjectData->bIsTimelineDirty)
+		{
+			// USUALLY SequenceID x < y implies x happened before y but we do not enforce this (so no optimization done here)
+			ObjectData->CachedTimelineSortedByStart.Sort([this, &Info](FSequenceId Left, FSequenceId Right)
+			{
+				const double LeftStartTime = FindSequenceData(Info.MakeSequenceInfo(Left))->Start;
+				const double RightStartTime = FindSequenceData(Info.MakeSequenceInfo(Right))->Start;
+				return LeftStartTime < RightStartTime;
+			});
+			ObjectData->bIsTimelineDirty = false;
+		}
+
+		// Find the index of first sequence that falls into the time window ...
+		const int32 LowerIndex = Algo::LowerBoundBy(ObjectData->CachedTimelineSortedByStart, Start, [this, Start, End, &Info](FSequenceId SequenceId)
+		{
+			const FPerSequenceData* SequenceData = FindSequenceData(Info.MakeSequenceInfo(SequenceId));
+			return Private::ShouldDisplayScopeInWindow(Start, End, SequenceData->Start, SequenceData->GetEndTime())
+				// Causes it to be included because LowerBoundBy looks <=. Also won't interfere with earlier sequences because LowerBoundBy looks for the FIRST element <=.
+				? FMath::Max(Start, SequenceData->Start)
+				// Causes it to be excluded
+				: SequenceData->Start;
+		});
+		if (LowerIndex >= ObjectData->CachedTimelineSortedByStart.Num())
+		{
+			return;
+		}
+
+		// ... and iterate until the first index is outside of the requested time window
+		for (int32 i = LowerIndex; i < ObjectData->CachedTimelineSortedByStart.Num(); ++i)
+		{
+			const FSequenceId SequenceId = ObjectData->CachedTimelineSortedByStart[i];
+			const FPerSequenceData* SequenceData = FindSequenceData(Info.MakeSequenceInfo(SequenceId));
+			
+			const bool bIsIndexInTimeWindow = SequenceData->Start <= End;
+			if (!bIsIndexInTimeWindow || Callback(SequenceId) == TraceServices::EEventEnumerate::Stop)
+			{
+				break;
+			}
+		}
+	}
+
+	TOptional<FVector2d> FProtocolMultiEndpointProvider::GetSequenceBounds(const FSequenceScopeInfo& Info) const
+	{
+		Session.ReadAccessCheck();
+		
+		const FPerSequenceData* SequenceData = FindSequenceData(Info);
+		if (SequenceData)
+		{
+			return FVector2d{ SequenceData->Start, SequenceData->GetEndTime() };
+		}
+
+		return {};
+	}
+
+	void FProtocolMultiEndpointProvider::EnumerateNetworkScopes(double Start, double End, const FSequenceScopeInfo& Info, TFunctionRef<TraceServices::EEventEnumerate(double, double, const FObjectNetworkScope&)> Callback) const
+	{
+		Session.ReadAccessCheck();
+
+		const FPerSequenceData* SequenceData = FindSequenceData(Info);
+		if (!SequenceData)
+		{
+		    return;
+		}
+
+		struct FNetworkScope
+		{
+		    double Start;
+		    double End;
+		    TOptional<FEndpointId> EndpointId;
+		};
+
+		// 1. Get all known scopes
+		TArray<FNetworkScope> Scopes;
+
+		// This is a very low number (number of analyzed .utrace files, usually <10) so it's relatively cheap to iterate through all events.
+		// The problem we're solving here is that we are analyzing events from multiple .utrace files in a non-deterministic order.
+		// That makes it tricky to analyse transmission scopes when the TransmissionStart, End, and Sink events are appended.
+		// It is easier just to aggregate them here...
+		for (int32 EventIdx = 0; EventIdx < SequenceData->NetworkScopeTimeline.GetEventCount(); ++EventIdx)
+		{
+			const double EventStart = SequenceData->NetworkScopeTimeline.GetEventStartTime(EventIdx);
+			const double EventEnd = SequenceData->NetworkScopeTimeline.GetEventEndTime(EventIdx);
+		    Scopes.Add({ EventStart, EventEnd, SequenceData->NetworkScopeTimeline.GetEvent(EventIdx) });
+		}
+		if (Scopes.IsEmpty())
+		{
+		    return;
+		}
+
+		// 2. Sort them by start time
+		Scopes.Sort([](const FNetworkScope& Left, const FNetworkScope& Right){ return Left.Start < Right.Start; });
+
+		// 3. Then fill in any gaps by interpreting them to be network transit times
+		const auto CallbackIfInRange = [Start, End, &Callback](double ScopeStart, double ScopeEnd, const FObjectNetworkScope& Event)
+		{
+		    return Private::ShouldDisplayScopeInWindow(Start, End, ScopeStart, ScopeEnd)
+		       ? Callback(ScopeStart, ScopeEnd, Event)
+		       : TraceServices::EEventEnumerate::Continue;
+		};
+
+		if (CallbackIfInRange(Scopes[0].Start, Scopes[0].End, { Scopes[0].EndpointId }) == TraceServices::EEventEnumerate::Stop)
+		{
+		    return;
+		}
+		for (int32 i = 1; i < Scopes.Num(); ++i)
+		{
+		    const FNetworkScope& Previous = Scopes[i - 1];
+		    const FNetworkScope& Current = Scopes[i];
+		    
+		    const double NetworkTransportTime = Current.Start - Previous.End;
+		    if (LIKELY(NetworkTransportTime >= 0.0) // If it is negative, we probably received bad data.
+		       && CallbackIfInRange(Previous.End, Current.Start, { {} }) == TraceServices::EEventEnumerate::Stop) 
+		    {
+		       return;
+		    }
+
+		    if (CallbackIfInRange(Current.Start, Current.End, { Current.EndpointId }) == TraceServices::EEventEnumerate::Stop)
+		    {
+		       return;
+		    }
+		}
+	}
+
+	void FProtocolMultiEndpointProvider::ReadProcessingStepTimeline(const FEndpointScopeInfo& Info, TFunctionRef<void(const TraceServices::ITimeline<FObjectProcessingStep>& Timeline)> Callback) const
+	{
+		Session.ReadAccessCheck();
+		
+		if (const FPerSequenceEndpointData* SequenceData = FindSequenceEndpointData(Info))
+		{
+			Callback(SequenceData->CpuTimeline);
+		}
+	}
+
+	void FProtocolMultiEndpointProvider::AppendInit(FInitMessage Message)
+	{
+		Private::LogInitMessage(TEXT("Received main trace file "), Session.GetTraceId(), Message);
+		
+		const TraceServices::FAnalysisSessionEditScope EditScope(Session);
+		if (!MainTraceInitData.IsSet())
+		{
+			MainTraceInitData = MoveTemp(Message);
+
+			// Now we know the main .utrace file's start time stamp, analysis is now able compute the time offsets to the other machines. So start the analysis now. 
+			Aggregator.StartAggregatedAnalysis();
+		}
+		else
+		{
+			// This can happen when you start recording in editor, stop, and start recording again.
+			// TODO: Recalculate time offsets.
+			UE_LOG(LogConcertInsights, Warning, TEXT("Received init message twice for main trace file %u. Investigate."), Session.GetTraceId());
+		}
+	}
+
+	bool FProtocolMultiEndpointProvider::ProcessGeneratedTraceData(float)
+	{
+		// This is run after Slate ticks: take a look at FUserInterfaceCommand::Run.
+		// That means that all ITimingViewExtender have already be run this tick.
+		// It would be better if this was run BEFORE but there's currently no integrated callback for that; FCoreDelegates::OnBeginFrame could be added to main to achieve this.
+		const TraceServices::FAnalysisSessionEditScope EditScope(Session);
+
+		Aggregator.EnumerateTraceFiles([this](uint64 TraceId)
+		{
+			Aggregator.ProcessEnqueuedData(TraceId, [this, TraceId](FProtocolDataQueue& DataQueue)
+			{
+				ProcessAggregatedTraceData(TraceId, DataQueue);
+			});
+			
+			return TraceServices::EEventEnumerate::Continue;
+		});
+
+		return true; 
+	}
+
+	void FProtocolMultiEndpointProvider::ProcessAggregatedTraceData(FEndpointId EndpointId, FProtocolDataQueue& DataQueue)
+	{
+		Session.WriteAccessCheck();
+
+		FProtocolQueuedItem Item;
+		while (DataQueue.MessageQueue.Dequeue(Item))
+		{
+			switch (Item.Type)
+			{
+			case EMessageType::Init: ProcessInit(EndpointId, Item.Message.Init); break;
+			case EMessageType::ObjectTraceBegin: ProcessObjectTraceBegin(EndpointId, Item.Message.ObjectTraceBegin); break;
+			case EMessageType::ObjectTraceEnd: ProcessObjectTraceEnd(EndpointId, Item.Message.ObjectTraceEnd); break;
+			case EMessageType::ObjectTransmissionStart: ProcessObjectTransmissionStart(EndpointId, Item.Message.TransmissionStart); break;
+			case EMessageType::ObjectTransmissionReceive: ProcessObjectTransmissionReceive(EndpointId, Item.Message.TransmissionReceive); break;
+			case EMessageType::ObjectSink: ProcessObjectSink(EndpointId, Item.Message.Sink); break;
+				
+			case EMessageType::None: [[fallthrough]];
+			default:
+				checkf(false, TEXT("Invalid type %u"), static_cast<uint8>(Item.Type));
+			}
+		}
+	}
+
+	void FProtocolMultiEndpointProvider::ProcessInit(FEndpointId EndpointId, const FInitMessage& Init)
+	{
+		const TraceServices::FAnalysisSessionEditScope EditScope(Session);
+		if (!ensure(!EndpointInitData.Contains(EndpointId)))
+		{
+			UE_LOG(LogConcertInsights, Warning, TEXT("Received init message twice for aggregated trace file %llu. Investigate."), EndpointId);
+			return;
+		}
+		
+		const bool bDuplicateEndpoint = Algo::AnyOf(EndpointInitData, [&Init](const TPair<FEndpointId, FInitMessage>& Entry)
+		{
+			const bool bSameEndpoint = Entry.Value.GetEndpointId().IsSet() == Init.GetEndpointId().IsSet() && *Entry.Value.GetEndpointId() == *Init.GetEndpointId();
+			const bool bDuplicateServer = Entry.Value.IsServer() == Init.IsServer();
+			return bSameEndpoint && bDuplicateServer;
+		});
+		if (bDuplicateEndpoint)
+		{
+			UE_CLOG(bDuplicateEndpoint, LogConcertInsights, Warning, TEXT("Session endpoint (endpoint:%s, bServer:%d) now encountered in trace file %llu is from was already encountered before. Investigate."),
+				Init.GetEndpointId() ? *(Init.GetEndpointId()->ToString()) : TEXT("none"),
+				Init.IsServer(),
+				EndpointId
+				);
+			UE_DEBUG_BREAK();
+			return;
+		}
+
+		Private::LogInitMessage(TEXT("Processed aggregated trace file "), EndpointId, Init);
+		EndpointInitData.Emplace(EndpointId, Init);
+	}
+
+	void FProtocolMultiEndpointProvider::ProcessObjectTraceBegin(FEndpointId EndpointId, const FObjectTraceMessage& Message)
+	{
+		const TraceServices::FAnalysisSessionEditScope EditScope(Session);
+		
+		const FEndpointScopeInfo Info { Message.Protocol, Message.ObjectPath, Message.SequenceId, EndpointId };
+		UpdateSequenceStats(Info, Message.Time);
+
+		if (const TOptional<double> ConvertedStart = ConvertEndpointCycleToTime(Info.EndpointId, Message.Time))
+		{
+			UE_LOG(LogConcertInsights, Verbose, TEXT("ObjectTraceBegin: Time: %f, EventName: %s, Context: %s"), *ConvertedStart, Message.EventName, *Info.ToString());
+			
+			FEndpointCpuTimeline& CpuTimeline = FindOrAddSequenceEndpointData(Info).CpuTimeline;
+			CpuTimeline.AppendBeginEvent(*ConvertedStart, { Message.EventName });
+		}
+	}
+
+	void FProtocolMultiEndpointProvider::ProcessObjectTraceEnd(FEndpointId EndpointId, const FObjectTraceMessage& Message)
+	{
+		const TraceServices::FAnalysisSessionEditScope EditScope(Session);
+		
+		const FEndpointScopeInfo Info { Message.Protocol, Message.ObjectPath, Message.SequenceId, EndpointId };
+		UpdateSequenceStats(Info, Message.Time);
+
+		if (const TOptional<double> ConvertedEnd = ConvertEndpointCycleToTime(Info.EndpointId, Message.Time))
+		{
+			UE_LOG(LogConcertInsights, Verbose, TEXT("ObjectTraceEnd: Time: %f, EventName: %s, Context: %s"), *ConvertedEnd, Message.EventName, *Info.ToString());
+			
+			FEndpointCpuTimeline& CpuTimeline = FindOrAddSequenceEndpointData(Info).CpuTimeline;
+			CpuTimeline.AppendEndEvent(*ConvertedEnd);
+		}
+	}
+
+	void FProtocolMultiEndpointProvider::ProcessObjectTransmissionStart(FEndpointId EndpointId, const FObjectTransmissionStartMessage& Message)
+	{
+		const TraceServices::FAnalysisSessionEditScope EditScope(Session);
+
+		const FEndpointScopeInfo Info { Message.Protocol, Message.ObjectPath, Message.SequenceId, EndpointId };
+		if (const TOptional<double> ConvertedTime = ConvertEndpointCycleToTime(Info.EndpointId, Message.GetTime()))
+		{
+			EndOpenNetworkScope(*ConvertedTime, Info);
+		}
+	}
+
+	void FProtocolMultiEndpointProvider::ProcessObjectTransmissionReceive(FEndpointId EndpointId, const FObjectTransmissionReceiveMessage& Message)
+	{
+		const TraceServices::FAnalysisSessionEditScope EditScope(Session);
+		
+		const TOptional<double> ConvertedTime = ConvertEndpointCycleToTime(EndpointId, Message.Time);
+		if (!ConvertedTime)
+		{
+			return;
+		}
+		
+		const FEndpointScopeInfo Info { Message.Protocol, Message.ObjectPath, Message.SequenceId, EndpointId };
+		UE_LOG(LogConcertInsights, Verbose, TEXT("ObjectTraceEnd: Time: %f, Context: %s"), *ConvertedTime, *Info.ToString());
+		UpdateSequenceStats_AlreadyConverted(Info, *ConvertedTime);
+
+		// UpdateSequenceStats allocates FPerSequenceData so Info should be mapped.
+		FPerSequenceData& SequenceData = FindSequenceDataChecked(Info);
+		const bool bOpenedScope = OpenNetworkScopeIfNotOpen(Info, *ConvertedTime, SequenceData);
+		UE_CLOG(!bOpenedScope, LogConcertInsights, Warning,
+			TEXT("A network scope was already open when processing %s. Was there a duplicate transmission receive or was the receive enqueued after a CPU scope was started?"), *Info.ToString()
+			);
+	}
+
+	void FProtocolMultiEndpointProvider::ProcessObjectSink(FEndpointId EndpointId, const FObjectSinkMessage& Message)
+	{
+		const TraceServices::FAnalysisSessionEditScope EditScope(Session);
+
+		const FEndpointScopeInfo Info { Message.Protocol, Message.ObjectPath, Message.SequenceId, EndpointId };
+		if (const TOptional<double> ConvertedTime = ConvertEndpointCycleToTime(Info.EndpointId, Message.Time))
+		{
+			EndOpenNetworkScope(*ConvertedTime, Info);
+			FindSequenceData(Info)->SinkData = { *ConvertedTime, EndpointId };
+		}
+	}
+	
+	bool FProtocolMultiEndpointProvider::OpenNetworkScopeIfNotOpen(const FEndpointScopeInfo& Info, double ConvertedTime, FPerSequenceData& SequenceData)
+	{
+		/*
+		 * Do not open a new network scope if the client that sent a sink event reports any further events.
+		 * This can happen like so:
+		 * {
+		 *		CONCERT_TRACE_REPLICATION_OBJECT_SCOPE
+		 *		CONCERT_TRACE_REPLICATION_OBJECT_SINK
+		 *		// CONCERT_TRACE_REPLICATION_OBJECT_SCOPE sends at end of scope, after the sink
+		 * }
+		 */
+		if (SequenceData.HasSequenceEnded() && SequenceData.SinkData->SinkEndpoint == Info.EndpointId)
+		{
+			return false;
+		}
+		
+		FPerSequenceEndpointData& EndpointData = FindOrAddSequenceEndpointData(Info);
+		const bool bHasOpenNetworkScope = EndpointData.LastScopeStartEventId.IsSet();
+		if (!bHasOpenNetworkScope)
+		{
+			const uint64 EventId = SequenceData.NetworkScopeTimeline.AppendBeginEvent(ConvertedTime, Info.EndpointId);
+			EndpointData.LastScopeStartEventId = EventId;
+			return true;
+		}
+		return false;
+	}
+	
+	void FProtocolMultiEndpointProvider::EndOpenNetworkScope(double ConvertedTime, const FEndpointScopeInfo& Info)
+	{
+		UE_LOG(LogConcertInsights, Verbose, TEXT("EndNetworkScope: Time: %f, Context: %s"), ConvertedTime, *Info.ToString());
+		UpdateSequenceStats_AlreadyConverted(Info, ConvertedTime);
+		
+		FPerSequenceEndpointData& EndpointData = FindOrAddSequenceEndpointData(Info);
+		const bool bHasOpenScope = EndpointData.LastScopeStartEventId.IsSet();
+		if (bHasOpenScope)
+		{
+			// UpdateSequenceStats allocates FPerSequenceData so Info should be mapped.
+			FPerSequenceData& SequenceData = FindSequenceDataChecked(Info);
+			SequenceData.NetworkScopeTimeline.EndEvent(*EndpointData.LastScopeStartEventId, ConvertedTime);
+			
+			EndpointData.LastScopeStartEventId.Reset();
+		}
+		else
+		{
+			UE_LOG(LogConcertInsights, Warning, TEXT("No network scope was open when processing %s. Was there a duplicate transmission start?"), *Info.ToString());
+		}
+	}
+
+	void FProtocolMultiEndpointProvider::UpdateSequenceStats(const FEndpointScopeInfo& Info, double Time)
+	{
+		const TOptional<double> ConvertedTime = ConvertEndpointCycleToTime(Info.EndpointId, Time);
+		if (ConvertedTime)
+		{
+			UpdateSequenceStats_AlreadyConverted(Info, *ConvertedTime);
+		}
+	}
+
+	void FProtocolMultiEndpointProvider::UpdateSequenceStats_AlreadyConverted(const FEndpointScopeInfo& Info, double Time)
+	{
+		ActiveProtocols.Add(Info.ProtocolId);
+		
+		// Make sure there is an entry for the object data
+		FPerObjectData& ObjectData = FindOrAddObjectData(Info);
+		
+		const bool bIsNewSequence = !PerSequenceData.Contains(Info);
+		if (bIsNewSequence)
+		{
+			ObjectData.CachedTimelineSortedByStart.Add(Info.SequenceId);
+			ObjectData.bIsTimelineDirty = true;
+		}
+		
+		FPerSequenceData& SequenceData = FindOrAddSequenceData(Info, Time);
+		const double NewStart = FMath::Min(SequenceData.Start, Time);
+		const bool bChangedStart = NewStart != SequenceData.Start;
+		ObjectData.bIsTimelineDirty |= bChangedStart;
+		
+		SequenceData.Start = NewStart;
+		SequenceData.Endpoints.Add(Info.EndpointId);
+
+		OpenNetworkScopeIfNotOpen(Info, Time, SequenceData);
+	}
+
+	TOptional<double> FProtocolMultiEndpointProvider::ConvertEndpointCycleToTime(FEndpointId Endpoint, double OtherEndpointTime) const
+	{
+		Session.ReadAccessCheck();
+		ensureMsgf(MainTraceInitData, TEXT("Analysis of other trace files is not supposed to have started, yet."));
+		
+		if (Endpoint == GetMainTraceId())
+		{
+			return OtherEndpointTime;
+		}
+		
+		if (const FInitMessage* OtherEndpoint = EndpointInitData.Find(Endpoint); OtherEndpoint)
+		{
+			// The other endpoint is the "source" timeline and we want that time relative to the main timeline (which acts as the "target" timeline)
+			return TimeSyncUtils::ConvertSourceToTargetTime(
+				MainTraceInitData->GetTraceInitTimeUtc(),
+				OtherEndpoint->GetTraceInitTimeUtc(),
+				MainTraceInitData->GetStartTime(),
+				OtherEndpoint->GetStartTime(),
+				OtherEndpointTime
+				);
+		}
+		else
+		{
+			UE_LOG(LogConcertInsights, Warning, TEXT("Endpoint %llu has no init data. Ignoring received event."), Endpoint);
+			return {};
+		}
+	}
+}
diff --git a/Engine/Plugins/Developer/Concert/ConcertInsights/ConcertInsightsVisualizer/Source/ConcertInsightsVisualizer/Private/Trace/ProtocolMultiEndpointProvider.h b/Engine/Plugins/Developer/Concert/ConcertInsights/ConcertInsightsVisualizer/Source/ConcertInsightsVisualizer/Private/Trace/ProtocolMultiEndpointProvider.h
new file mode 100644
--- /dev/null
+++ b/Engine/Plugins/Developer/Concert/ConcertInsights/ConcertInsightsVisualizer/Source/ConcertInsightsVisualizer/Private/Trace/ProtocolMultiEndpointProvider.h
@@ -0,0 +1,247 @@
+// Copyright Epic Games, Inc. All Rights Reserved.
+
+#pragma once
+
+#include "Analysis/IProtocolDataTarget.h"
+#include "Messages/InitMessage.h"
+#include "Messages/ObjectSinkMessage.h"
+#include "Messages/ObjectTraceMessage.h"
+#include "Messages/ObjectTransmissionReceiveMessage.h"
+#include "Messages/ObjectTransmissionStartMessage.h"
+#include "ObjectSequence/ObjectNetworkScope.h"
+#include "ObjectSequence/ObjectProcessingStep.h"
+#include "TraceAggregator.h"
+#include "Types/ProtocolId.h"
+#include "Types/ScopeInfos.h"
+#include "Types/SequenceId.h"
+
+#include "Containers/Ticker.h"
+#include "HAL/Platform.h"
+#include "Model/IntervalTimeline.h"
+#include "Model/MonotonicTimeline.h"
+#include "Model/PointTimeline.h"
+#include "Templates/Function.h"
+#include "TraceServices/Model/AnalysisSession.h"
+
+#include <limits>
+
+namespace UE::ConcertInsightsVisualizer
+{
+	class FSingleProtocolEndpointProvider;
+	
+	/**
+	 * This is the final provider that is visualized in Insights.
+	 * 
+	 * It aggregates multiple FSingleProtocolEndpointProvider into timelines that are then displayed.
+	 * There is a main provider, which is the one for which the .utrace is opened, to which all other traces are made relative.
+	 *
+	 * This is what the data will be visualized as in the Insights graph:
+	 * [--------------------------------------------- FObjectSequence - ActorName - Sequence 1 --------------------------------------------]
+	 * [--------- FObjectNetworkScope - Client 1---------][FObjectNetworkScope - Transit][--------- FObjectNetworkScope - Client 2---------]
+	 * [FObjectProcessingStep 1] [FObjectProcessingStep 2]                               [FObjectProcessingStep 1] [FObjectProcessingStep 2]
+	 * Sequence 2 would go under this, etc.
+	 */
+	class FProtocolMultiEndpointProvider
+		: public TraceServices::IProvider
+		, public TraceServices::IEditableProvider
+		, public IProtocolDataTarget
+	{
+	public:
+		
+		/** ID used for registration with Insights. */
+		static const FName ProviderName;
+
+		FProtocolMultiEndpointProvider(TraceServices::IAnalysisSession& Session UE_LIFETIMEBOUND);
+		virtual ~FProtocolMultiEndpointProvider() override;
+
+		/** Utility for getting the display name of a particular endpoint. Returns nullptr if invalid. */
+		const TCHAR* GetEndpointDisplayName(FEndpointId EndpointId) const;
+
+		/** Gets all protocols for which objects were traced. */
+		void EnumerateProtocols(TFunctionRef<TraceServices::EEventEnumerate(FProtocolId)> Callback) const;
+		/** Gets all objects for which data was traced. */
+		void EnumerateObjects(FProtocolId Protocol, TFunctionRef<TraceServices::EEventEnumerate(FObjectPath Object)> Callback) const;
+		/** Gets all endpoints that participate in the given sequence. */
+		void EnumerateEndpointsInSequence(const FSequenceScopeInfo& Info, TFunctionRef<TraceServices::EEventEnumerate(FEndpointId EndpointId)> Callback) const;
+		
+		/** 1st row. Gets the timeline of all sequences for a particular object. */
+		void EnumerateSequences(double Start, double End, const FObjectScopeInfo&, TFunctionRef<TraceServices::EEventEnumerate(FSequenceId Sequence)> Callback) const;
+
+		/** @return Start and end time of the sequence relative to the main trace file. */
+		TOptional<FVector2d> GetSequenceBounds(const FSequenceScopeInfo& Info) const;
+		
+		/** 2nd row. Lists out all network scopes that are active in the given time range. A scope is a time period in which an endpoint is doing processing or the object is in transit. */
+		void EnumerateNetworkScopes(double Start, double End, const FSequenceScopeInfo& Info, TFunctionRef<TraceServices::EEventEnumerate(double StartTime, double EndTime, const FObjectNetworkScope& NetworkScope)> Callback) const;
+
+		/** 3rd row+. Gets the timeline of all processing steps performed by a specific endpoint for a particular object sequence. */
+		void ReadProcessingStepTimeline(const FEndpointScopeInfo& Info, TFunctionRef<void(const TraceServices::ITimeline<FObjectProcessingStep>& Timeline)> Callback) const;
+
+		//~ Begin IProtocolDataTarget Interface
+		virtual void AppendInit(FInitMessage Message) override;
+		virtual void AppendObjectTraceBegin(FObjectTraceBeginMessage Message) override { ProcessObjectTraceBegin(Session.GetTraceId(), MoveTemp(Message)); }
+		virtual void AppendObjectTraceEnd(FObjectTraceEndMessage Message) override { ProcessObjectTraceEnd(Session.GetTraceId(), MoveTemp(Message)); }
+		virtual void AppendObjectTransmissionStart(FObjectTransmissionStartMessage Message) override { ProcessObjectTransmissionStart(Session.GetTraceId(), MoveTemp(Message)); }
+		virtual void AppendObjectTransmissionReceive(FObjectTransmissionReceiveMessage Message) override { ProcessObjectTransmissionReceive(Session.GetTraceId(), MoveTemp(Message)); }
+		virtual void AppendObjectSink(FObjectSinkMessage Message) override { ProcessObjectSink(Session.GetTraceId(), MoveTemp(Message)); }
+		//~ Begin IProtocolDataTarget Interface
+
+	private:
+
+		using FEndpointCpuTimeline = TraceServices::TMonotonicTimeline<FObjectProcessingStep>;
+		using FNetworkScopeTimeline = TraceServices::TIntervalTimeline<FEndpointId>;
+		using FSequenceEventTimeline = TraceServices::TPointTimeline<FSequenceId>;
+		
+		struct FPerSequenceEndpointData
+		{
+			/**
+			 * Timelines where endpoints spent CPU time.
+			 * Displayed as 3rd row for a sequence.
+			 */
+			FEndpointCpuTimeline CpuTimeline;
+
+			/**
+			 * Result of FNetworkScopeTimeline::AppendBeginEvent for FPerSequenceData::NetworkScopeTimeline.
+			 * Set when a scope has been started, unset when it has ended.
+			 */
+			TOptional<uint64> LastScopeStartEventId;
+
+			FPerSequenceEndpointData(TraceServices::ILinearAllocator& Allocator)
+				: CpuTimeline(Allocator)
+			{}
+		};
+
+		struct FSequenceEndData
+		{
+			/** The time of the sequence event with the latest time processed so far. Set once the sink event has been received; the track will display an infinite time until then. */
+			double End;
+			/** The endpoint for which the sequence ended. */
+			FEndpointId SinkEndpoint;
+		};
+
+		struct FPerSequenceData
+		{
+			/** The time of the sequence event with the earliest time processed so far; updated as new events come in.  */
+			double Start;
+			/** Set when the sequence sink has been encountered */
+			TOptional<FSequenceEndData> SinkData;
+
+			/** Endpoints that participated in this sequence */
+			TSet<FEndpointId> Endpoints;
+
+			/**
+			 * Indicates the intervals at which an object was processed by an endpoint.
+			 * Does NOT contain intervals for transit scopes; those are determined implicitly in EnumerateNetworkScopes.
+			 */
+			FNetworkScopeTimeline NetworkScopeTimeline;
+
+			/** Gets the end time as it should be displayed: infinite if pending and the actual end time if the sequence has ended. */
+			double GetEndTime() const { return SinkData ? SinkData->End : std::numeric_limits<double>::infinity(); }
+			/** Whether a sink event has been received for the sequence. */
+			bool HasSequenceEnded() const { return SinkData.IsSet(); }
+
+			FPerSequenceData(TraceServices::ILinearAllocator& Allocator, double Start)
+				: Start(Start)
+				, NetworkScopeTimeline(Allocator)
+			{}
+		};
+		
+		struct FPerObjectData
+		{
+			/**
+			 * When bIsTimelineDirty == false, this contains all sequences sorted by their start time.
+			 * 
+			 * Every time new sequence is added or the sequence bounds change (FPerSequenceData::Start or FPerSequenceData::End), this needs to be resorted.
+			 * This is resorted lazily, i.e. only when EnumerateSequences is called.
+			 *
+			 * Context:
+			 * - USUALLY sequences start a monotonous, i.e. id x < y implies that sequence x started before y, but we do not want to enforce this to make it easier for the trace API user.
+			 * - OFTEN sequence are processed in order (i.e. usually we can just CachedTimelineSortedByStart.Add new sequence) but they can be processed out of order.
+			 * - Asymptotically, it is cheaper to dirty this array with with an O(1) update and do one full O(nlogn) resort rather than keeping this
+			 *	array sorted making each update O(n), which effectively boils down to insertion sort of O(n^2). Whether this is also better in practice is not profiled.
+			 * For the above 3 points in mind, we opt for a full resort in EnumerateSequences.
+			 *
+			 * Mutable because const function EnumerateSequences updates this.
+			 * In the perspective of the API user, the internal state does not change.
+			 */
+			mutable TArray<FSequenceId> CachedTimelineSortedByStart;
+
+			/**
+			 * When true, CachedTimelineSortedByStart needs to be re-sorted.
+			 * This flag is set when either 1. a new sequence ID is added, or 2. a FPerSequenceData::Start or FPerSequenceData::End time value is changed.
+			 *
+			 * Mutable because const function EnumerateSequences updates this.
+			 * In the perspective of the API user, the internal state does not change.
+			 */
+			mutable bool bIsTimelineDirty = false;
+		};
+
+		/** The main session this provider is for. Outlives this object. */
+		TraceServices::IAnalysisSession& Session;
+		/** Traces related files and exposes their data. Set once the session starts. */
+		FTraceAggregator Aggregator;
+		/** Handle to ProcessGeneratedTraceData. */
+		const FTSTicker::FDelegateHandle TickHandle;
+
+		/** Keeps track of all the protocols encountered so far. */
+		TSet<FProtocolId> ActiveProtocols;
+		
+		TMap<FObjectScopeInfo, FPerObjectData> PerObjectData;
+		TMap<FSequenceScopeInfo, FPerSequenceData> PerSequenceData;
+		TMap<FEndpointScopeInfo, FPerSequenceEndpointData> PerSequenceEndpointData;
+
+		/** Contains the init data found in the Session's trace. */
+		TOptional<FInitMessage> MainTraceInitData;
+		/** Maps .utrace file ID to the init data found in that file. The endpoint ID is the same as the .utrace file ID. Does not contain MainTraceInitData. */
+		TMap<FEndpointId, FInitMessage> EndpointInitData;
+
+		uint64 GetMainTraceId() const { return Session.GetTraceId(); }
+
+		/** Process the data from Aggregator. This is run after Slate ticks: take a look at FUserInterfaceCommand::Run. */
+		bool ProcessGeneratedTraceData(float DeltaTime);
+		void ProcessAggregatedTraceData(FEndpointId EndpointId, FProtocolDataQueue& DataQueue);
+		
+		void ProcessInit(FEndpointId EndpointId, const FInitMessage& Init);
+		void ProcessObjectTraceBegin(FEndpointId EndpointId, const FObjectTraceMessage& Message);
+		void ProcessObjectTraceEnd(FEndpointId EndpointId, const FObjectTraceMessage& Message);
+		void ProcessObjectTransmissionStart(FEndpointId EndpointId, const FObjectTransmissionStartMessage& Message);
+		void ProcessObjectTransmissionReceive(FEndpointId EndpointId, const FObjectTransmissionReceiveMessage& Message);
+		void ProcessObjectSink(FEndpointId EndpointId, const FObjectSinkMessage& Message);
+
+		// Getters
+		const FPerObjectData* FindObjectData(const FObjectScopeInfo& Info) const { return PerObjectData.Find(Info); }
+		const FPerSequenceData* FindSequenceData(const FSequenceScopeInfo& Info) const { return PerSequenceData.Find(Info); }
+		FPerSequenceData* FindSequenceData(const FSequenceScopeInfo& Info) { return PerSequenceData.Find(Info); }
+		const FPerSequenceEndpointData* FindSequenceEndpointData(const FEndpointScopeInfo& Info) const { return PerSequenceEndpointData.Find(Info); }
+		// Setters
+		FPerObjectData& FindOrAddObjectData(const FObjectScopeInfo& Info) { return PerObjectData.FindOrAdd(Info); }
+		FPerSequenceData& FindOrAddSequenceData(const FSequenceScopeInfo& Info, double Start) { return PerSequenceData.FindOrAdd(Info, FPerSequenceData(Session.GetLinearAllocator(), Start)); }
+		FPerSequenceData& FindSequenceDataChecked(const FSequenceScopeInfo& Info) { return PerSequenceData[Info]; }
+		FPerSequenceEndpointData& FindOrAddSequenceEndpointData(const FEndpointScopeInfo& Info) { return PerSequenceEndpointData.FindOrAdd(Info, FPerSequenceEndpointData(Session.GetLinearAllocator())); }
+
+		/**
+		 * Opens a network scope.
+		 * 
+		 * Usually opened by TransmissionReceive.
+		 * 
+		 * Also used whenever any message is received to check whether a scope was opened implicitly:
+		 * At the beginning of a sequence, there is usually no TransmissionReceive but a scope needs to be opened.
+		 *
+		 * @return Whether a scope was opened
+		 */
+		bool OpenNetworkScopeIfNotOpen(const FEndpointScopeInfo& Info, double ConvertedTime, FPerSequenceData& SequenceData);
+		/** Used by TransmissionStart or ObjectSink event to end the scope of the current network user. */
+		void EndOpenNetworkScope(double ConvertedTime, const FEndpointScopeInfo& Info);
+		
+		/**
+		 * Converts Time using ConvertEndpointCycleToTime, updates ActiveProtocols, updates the known start and end times of the sequence,
+		 * and caches that events occured at the times (for speeding up search which events happened in a certain time window).
+		 */
+		void UpdateSequenceStats(const FEndpointScopeInfo& Info, double Time);
+		/** Overload that already accepts an already converted Time (ConvertEndpointCycleToTime was already called, just store the Time that is passed in). */
+		void UpdateSequenceStats_AlreadyConverted(const FEndpointScopeInfo& Info, double Time);
+		
+		/** Computes what time on the main .utrace file another endpoint's time corresponds to. Fails if that endpoint did not send any init event. */
+		TOptional<double> ConvertEndpointCycleToTime(FEndpointId Endpoint, double OtherEndpointTime) const;
+	};
+}
+
diff --git a/Engine/Plugins/Developer/Concert/ConcertInsights/ConcertInsightsVisualizer/Source/ConcertInsightsVisualizer/Private/Trace/TraceAggregator.cpp b/Engine/Plugins/Developer/Concert/ConcertInsights/ConcertInsightsVisualizer/Source/ConcertInsightsVisualizer/Private/Trace/TraceAggregator.cpp
new file mode 100644
--- /dev/null
+++ b/Engine/Plugins/Developer/Concert/ConcertInsights/ConcertInsightsVisualizer/Source/ConcertInsightsVisualizer/Private/Trace/TraceAggregator.cpp
@@ -0,0 +1,96 @@
+// Copyright Epic Games, Inc. All Rights Reserved.
+
+#include "TraceAggregator.h"
+
+#include "LogConcertInsights.h"
+#include "Trace/StoreClient.h"
+#include "TraceServices/AnalyzerFactories.h"
+
+namespace UE::ConcertInsightsVisualizer
+{
+	FTraceAggregator::FTraceData::FTraceData(TSharedRef<TraceServices::IAnalysisSession> InSession)
+		: Session(MoveTemp(InSession))
+		, Analyzer(*Session, DataQueue)
+	{
+		AnalysisContext.AddAnalyzer(Analyzer);
+	}
+
+	FTraceAggregator::FTraceAggregator(const Trace::FStoreClient& StoreClient, uint64 MainTraceId)
+		: StoreClient(StoreClient)
+		, MainTraceId(MainTraceId)
+	{}
+
+	void FTraceAggregator::StartAggregatedAnalysis()
+	{
+		const Trace::FStoreClient::FTraceInfo* MainTraceInfo = StoreClient.GetTraceInfoById(MainTraceId);
+		if (ensure(MainTraceInfo))
+		{
+			// For now, only non-live traces can be analyzed because there is no API for synchronizing the analyzing threads: reading from the providers would cause race conditions.
+			AnalyzeCompletedTraces(*MainTraceInfo);
+		}
+	}
+
+	void FTraceAggregator::EnumerateTraceFiles(TFunctionRef<TraceServices::EEventEnumerate(uint64 TraceId)> Callback) const
+	{
+		for (const TPair<uint64, TUniquePtr<FTraceData>>& Pair : AggregatedTraces)
+		{
+			if (Callback(Pair.Key) == TraceServices::EEventEnumerate::Stop)
+			{
+				break;
+			}
+		}
+	}
+
+	void FTraceAggregator::ProcessEnqueuedData(uint64 TraceId, TFunctionRef<void(FProtocolDataQueue& DataQueue)> Callback)
+	{
+		if (const TUniquePtr<FTraceData>* TraceData = AggregatedTraces.Find(TraceId))
+		{
+			Callback(TraceData->Get()->DataQueue);
+		}
+	}
+
+	void FTraceAggregator::AnalyzeCompletedTraces(const FTraceCachedInfo& MainTraceInfo)
+	{
+		UE_LOG(LogConcertInsights, Log, TEXT("Analyzing aggregated traces"));
+			
+		for (uint32 TraceIndex = 0; TraceIndex < StoreClient.GetTraceCount(); ++TraceIndex)
+		{
+			const Trace::FStoreClient::FTraceInfo* TraceInfo = StoreClient.GetTraceInfo(TraceIndex);
+			if (ensure(TraceInfo) && TraceInfo->GetId() != MainTraceInfo.Id && ShouldAnalyzeTrace(*TraceInfo, MainTraceInfo))
+			{
+				AnalyzeTrace(TraceInfo->GetId());
+			}
+		}
+	}
+
+	bool FTraceAggregator::ShouldAnalyzeTrace(const Trace::FStoreClient::FTraceInfo& ConsideredTraceInfo, const FTraceCachedInfo& MainTraceInfo)
+	{
+		check(ConsideredTraceInfo.GetId() != MainTraceInfo.Id);
+		
+		const FDateTime ConsideredTime(ConsideredTraceInfo.GetTimestamp());
+		const FDateTime MainTraceTime(MainTraceInfo.Timestamp);
+
+		// TODO DP: This is a hacky way of finding related traces. We assume they were recorded in sync.
+		// In the future (maybe 5.6), we must get them by session ID but the API for that is not ready, yet.
+		const FTimespan TimeDifference = ConsideredTime <= MainTraceTime
+			? MainTraceTime - ConsideredTime
+			: ConsideredTime - MainTraceTime;
+		return TimeDifference <= FTimespan::FromSeconds(5);
+	}
+
+	void FTraceAggregator::AnalyzeTrace(uint64 TraceId)
+	{
+		TUniquePtr<Trace::IInDataStream> TraceStream = StoreClient.ReadTrace(TraceId);
+		if (!ensure(TraceStream))
+		{
+			return;
+		}
+		UE_LOG(LogConcertInsights, Log, TEXT("Starting analysis of aggregated trace %llu"), TraceId);
+		
+		Trace::IInDataStream* StreamNakedPointer = TraceStream.Get();
+		const TSharedPtr<TraceServices::IAnalysisSession> Session = TraceServices::CreateAnalysisSession(TraceId, TEXT("FTraceAggregator"), MoveTemp(TraceStream));
+		
+		TUniquePtr<FTraceData>& TraceData = AggregatedTraces.Emplace(TraceId, MakeUnique<FTraceData>(Session.ToSharedRef()));
+		TraceData->AnalysisProcessor = TraceData->AnalysisContext.Process(*StreamNakedPointer);
+	}
+}
diff --git a/Engine/Plugins/Developer/Concert/ConcertInsights/ConcertInsightsVisualizer/Source/ConcertInsightsVisualizer/Private/Trace/TraceAggregator.h b/Engine/Plugins/Developer/Concert/ConcertInsights/ConcertInsightsVisualizer/Source/ConcertInsightsVisualizer/Private/Trace/TraceAggregator.h
new file mode 100644
--- /dev/null
+++ b/Engine/Plugins/Developer/Concert/ConcertInsights/ConcertInsightsVisualizer/Source/ConcertInsightsVisualizer/Private/Trace/TraceAggregator.h
@@ -0,0 +1,90 @@
+// Copyright Epic Games, Inc. All Rights Reserved.
+
+#pragma once
+
+#include "Analysis/ProtocolDataQueue.h"
+#include "Analysis/ProtocolEndpointAnalyzer.h"
+
+#include "HAL/Platform.h"
+#include "Templates/Function.h"
+#include "Templates/UnrealTemplate.h"
+#include "Trace/Analysis.h"
+#include "Trace/StoreClient.h"
+#include "TraceServices/Containers/Timelines.h"
+#include "TraceServices/Model/AnalysisSession.h"
+
+namespace UE::Trace{ class FStoreClient; }
+
+namespace UE::ConcertInsightsVisualizer
+{
+	struct FTraceCachedInfo
+	{
+		uint64 Id;
+		uint64 Timestamp;
+		FTraceCachedInfo(const Trace::FStoreClient::FTraceInfo& Info)
+		: Id(Info.GetId())
+		, Timestamp(Info.GetTimestamp())
+		{}
+	};
+	
+	/** Analyses trace files that relate to a main .utrace. */
+	class FTraceAggregator : public FNoncopyable
+	{
+	public:
+
+		FTraceAggregator(const Trace::FStoreClient& StoreClient UE_LIFETIMEBOUND, uint64 MainTraceId);
+
+		/**
+		 * Finds related trace files and kicks of the analyzing threads.
+		 *
+		 * Data will be enqueued so it can be read from the main thread.
+		 * You are expected to call ProcessEnqueuedData at regular intervals from the main thread and transfer data, e.g. every tick.
+		 */
+		void StartAggregatedAnalysis();
+
+		/** Enumerates all trace files currently being analyzed. */
+		void EnumerateTraceFiles(TFunctionRef<TraceServices::EEventEnumerate(uint64 TraceId)> Callback) const;
+		
+		/** Gets the data that has been produced since the last time. */
+		void ProcessEnqueuedData(uint64 TraceId, TFunctionRef<void(FProtocolDataQueue& DataQueue)> Callback);
+
+	private:
+
+		struct FTraceData
+		{
+			/**
+			 * Session created for this trace file.
+			 * Does not any processing (FAnalysisSession::Start not called) but simply exists for storing strings.
+			 * Owns the underlying IInDataStream and keeps it alive.
+			 */
+			TSharedRef<TraceServices::IAnalysisSession> Session;
+
+			/** Holds the analyzed data. This is actively dequeued */
+			FProtocolDataQueue DataQueue;
+			/** Fills Provider with data */
+			FProtocolEndpointAnalyzer Analyzer;
+
+			/** Context that instruments the analysis */
+			Trace::FAnalysisContext AnalysisContext;
+			/** Handle to thread doing the processing */
+			Trace::FAnalysisProcessor AnalysisProcessor;
+
+			FTraceData(TSharedRef<TraceServices::IAnalysisSession> InSession);
+		};
+
+		/** Client from which to obtain related trace files. */
+		const Trace::FStoreClient& StoreClient;
+		/** ID of the trace being aggregated into. */
+		const uint64 MainTraceId;
+
+		/** The files being analyzed */
+		TMap<uint64, TUniquePtr<FTraceData>> AggregatedTraces;
+		
+		void AnalyzeCompletedTraces(const FTraceCachedInfo& MainTraceInfo);
+
+		static bool ShouldAnalyzeTrace(const Trace::FStoreClient::FTraceInfo& ConsideredTraceInfo, const FTraceCachedInfo& MainTraceInfo);
+		void AnalyzeTrace(uint64 TraceId);
+	};
+}
+
+
diff --git a/Engine/Plugins/Developer/Concert/ConcertInsights/ConcertInsightsVisualizer/Source/ConcertInsightsVisualizer/Private/Trace/Types/EndpointId.h b/Engine/Plugins/Developer/Concert/ConcertInsights/ConcertInsightsVisualizer/Source/ConcertInsightsVisualizer/Private/Trace/Types/EndpointId.h
new file mode 100644
--- /dev/null
+++ b/Engine/Plugins/Developer/Concert/ConcertInsights/ConcertInsightsVisualizer/Source/ConcertInsightsVisualizer/Private/Trace/Types/EndpointId.h
@@ -0,0 +1,11 @@
+// Copyright Epic Games, Inc. All Rights Reserved.
+
+#pragma once
+
+#include "HAL/Platform.h"
+
+namespace UE::ConcertInsightsVisualizer
+{
+	/** Id of a network endpoint, which can do processing of an object. */
+	using FEndpointId = uint64;
+}
\ No newline at end of file
diff --git a/Engine/Plugins/Developer/Concert/ConcertInsights/ConcertInsightsVisualizer/Source/ConcertInsightsVisualizer/Private/Trace/Types/ProtocolId.h b/Engine/Plugins/Developer/Concert/ConcertInsights/ConcertInsightsVisualizer/Source/ConcertInsightsVisualizer/Private/Trace/Types/ProtocolId.h
new file mode 100644
--- /dev/null
+++ b/Engine/Plugins/Developer/Concert/ConcertInsights/ConcertInsightsVisualizer/Source/ConcertInsightsVisualizer/Private/Trace/Types/ProtocolId.h
@@ -0,0 +1,10 @@
+// Copyright Epic Games, Inc. All Rights Reserved.
+
+#pragma once
+
+#include "HAL/Platform.h"
+
+namespace UE::ConcertInsightsVisualizer
+{
+	using FProtocolId = uint8;
+}
\ No newline at end of file
diff --git a/Engine/Plugins/Developer/Concert/ConcertInsights/ConcertInsightsVisualizer/Source/ConcertInsightsVisualizer/Private/Trace/Types/ScopeInfos.cpp b/Engine/Plugins/Developer/Concert/ConcertInsights/ConcertInsightsVisualizer/Source/ConcertInsightsVisualizer/Private/Trace/Types/ScopeInfos.cpp
new file mode 100644
--- /dev/null
+++ b/Engine/Plugins/Developer/Concert/ConcertInsights/ConcertInsightsVisualizer/Source/ConcertInsightsVisualizer/Private/Trace/Types/ScopeInfos.cpp
@@ -0,0 +1 @@
+#include "ScopeInfos.h"
diff --git a/Engine/Plugins/Developer/Concert/ConcertInsights/ConcertInsightsVisualizer/Source/ConcertInsightsVisualizer/Private/Trace/Types/ScopeInfos.h b/Engine/Plugins/Developer/Concert/ConcertInsights/ConcertInsightsVisualizer/Source/ConcertInsightsVisualizer/Private/Trace/Types/ScopeInfos.h
new file mode 100644
--- /dev/null
+++ b/Engine/Plugins/Developer/Concert/ConcertInsights/ConcertInsightsVisualizer/Source/ConcertInsightsVisualizer/Private/Trace/Types/ScopeInfos.h
@@ -0,0 +1,81 @@
+// Copyright Epic Games, Inc. All Rights Reserved.
+
+#pragma once
+
+#include "EndpointId.h"
+#include "ProtocolId.h"
+#include "SequenceId.h"
+#include "Trace/Analysis.h"
+
+namespace UE::ConcertInsightsVisualizer
+{
+	class FSingleProtocolEndpointProvider;
+	using FObjectPath = FStringView;
+
+	/** Identifies all data associated with an object */
+	struct FObjectScopeInfo
+	{
+		FProtocolId ProtocolId;
+		FObjectPath ObjectPath;
+
+		FObjectScopeInfo() = default;
+		FObjectScopeInfo(FProtocolId ProtocolId, const FObjectPath& ObjectPath)
+			: ProtocolId(ProtocolId)
+			, ObjectPath(ObjectPath)
+		{}
+		
+		struct FSequenceScopeInfo MakeSequenceInfo(FSequenceId SequenceId) const;
+
+		friend bool operator==(const FObjectScopeInfo& Left, const FObjectScopeInfo& Right) { return Left.ProtocolId == Right.ProtocolId && Left.ObjectPath == Right.ObjectPath; }
+		friend bool operator!=(const FObjectScopeInfo& Left, const FObjectScopeInfo& Right) { return !(Left == Right); }
+		friend uint32 GetTypeHash(const FObjectScopeInfo& Info) { return HashCombine(GetTypeHash(Info.ProtocolId), GetTypeHash(Info.ObjectPath)); }
+	};
+	
+	/** Identifies a sequence of an object */
+	struct FSequenceScopeInfo : FObjectScopeInfo
+	{
+		FSequenceId SequenceId;
+
+		FSequenceScopeInfo() = default;
+		FSequenceScopeInfo(FProtocolId ProtocolId, const FObjectPath& ObjectPath, FSequenceId SequenceId)
+			: FObjectScopeInfo(ProtocolId, ObjectPath)
+			, SequenceId(SequenceId)
+		{}
+
+		struct FEndpointScopeInfo MakeEndpointInfo(FEndpointId EndpointId) const;
+
+		friend bool operator==(const FSequenceScopeInfo& Left, const FSequenceScopeInfo& Right) { return static_cast<const FObjectScopeInfo&>(Left) == static_cast<const FObjectScopeInfo&>(Left) && Left.SequenceId == Right.SequenceId; }
+		friend bool operator!=(const FSequenceScopeInfo& Left, const FSequenceScopeInfo& Right) { return !(Left == Right); }
+		friend uint32 GetTypeHash(const FSequenceScopeInfo& Info) { return HashCombine(GetTypeHash(Info.SequenceId), GetTypeHash(static_cast<const FObjectScopeInfo&>(Info))); }
+	};
+	
+	/** Identifies an endpoint that did processing in a sequence. */
+	struct FEndpointScopeInfo : FSequenceScopeInfo
+	{
+		FEndpointScopeInfo() = default;
+		FEndpointScopeInfo(FProtocolId ProtocolId, const FObjectPath& ObjectPath, FSequenceId SequenceId, FEndpointId EndpointId)
+			: FSequenceScopeInfo(ProtocolId, ObjectPath, SequenceId)
+			, EndpointId(EndpointId)
+		{}
+
+		FEndpointId EndpointId;
+
+		FString ToString() const { return FString::Printf(TEXT("Protocol: %u, Object: %s, Sequence: %u, Endpoint: %llu"), ProtocolId, ObjectPath.GetData(), SequenceId, EndpointId); }
+		
+		friend bool operator==(const FEndpointScopeInfo& Left, const FEndpointScopeInfo& Right) { return static_cast<const FSequenceScopeInfo&>(Left) == static_cast<const FSequenceScopeInfo&>(Right) && Left.EndpointId == Right.EndpointId; }
+		friend bool operator!=(const FEndpointScopeInfo& Left, const FEndpointScopeInfo& Right) { return !(Left == Right); }
+		friend uint32 GetTypeHash(const FEndpointScopeInfo& Info) { return HashCombine(GetTypeHash(Info.EndpointId), GetTypeHash(static_cast<const FSequenceScopeInfo&>(Info))); }
+	};
+
+	
+	FORCEINLINE FSequenceScopeInfo FObjectScopeInfo::MakeSequenceInfo(FSequenceId SequenceId) const
+	{
+		return { ProtocolId, ObjectPath, SequenceId };
+	}
+	
+	FORCEINLINE FEndpointScopeInfo FSequenceScopeInfo::MakeEndpointInfo(FEndpointId EndpointId) const
+	{
+		return { ProtocolId, ObjectPath, SequenceId, EndpointId };
+	}
+}
+
diff --git a/Engine/Plugins/Developer/Concert/ConcertInsights/ConcertInsightsVisualizer/Source/ConcertInsightsVisualizer/Private/Trace/Types/SequenceEndpointPair.h b/Engine/Plugins/Developer/Concert/ConcertInsights/ConcertInsightsVisualizer/Source/ConcertInsightsVisualizer/Private/Trace/Types/SequenceEndpointPair.h
new file mode 100644
--- /dev/null
+++ b/Engine/Plugins/Developer/Concert/ConcertInsights/ConcertInsightsVisualizer/Source/ConcertInsightsVisualizer/Private/Trace/Types/SequenceEndpointPair.h
@@ -0,0 +1,19 @@
+// Copyright Epic Games, Inc. All Rights Reserved.
+
+#pragma once
+
+#include "EndpointId.h"
+#include "SequenceId.h"
+
+namespace UE::ConcertInsightsVisualizer
+{
+	struct FSequenceEndpointPair
+	{
+		FSequenceId SequenceId;
+		FEndpointId EndpointId;
+
+		friend bool operator==(const FSequenceEndpointPair& Left, const FSequenceEndpointPair& Right) { return Left.SequenceId == Right.SequenceId && Left.EndpointId == Right.EndpointId; }
+		friend bool operator!=(const FSequenceEndpointPair& Left, const FSequenceEndpointPair& Right) { return !(Left == Right); }
+		friend uint32 GetTypeHash(const FSequenceEndpointPair& Pair) { return HashCombine(Pair.SequenceId, Pair.EndpointId); }
+	};
+}
diff --git a/Engine/Plugins/Developer/Concert/ConcertInsights/ConcertInsightsVisualizer/Source/ConcertInsightsVisualizer/Private/Trace/Types/SequenceId.h b/Engine/Plugins/Developer/Concert/ConcertInsights/ConcertInsightsVisualizer/Source/ConcertInsightsVisualizer/Private/Trace/Types/SequenceId.h
new file mode 100644
--- /dev/null
+++ b/Engine/Plugins/Developer/Concert/ConcertInsights/ConcertInsightsVisualizer/Source/ConcertInsightsVisualizer/Private/Trace/Types/SequenceId.h
@@ -0,0 +1,10 @@
+// Copyright Epic Games, Inc. All Rights Reserved.
+
+#pragma once
+
+#include "HAL/Platform.h"
+
+namespace UE::ConcertInsightsVisualizer
+{
+	using FSequenceId = uint32;
+}
\ No newline at end of file
diff --git a/Engine/Plugins/Developer/Concert/ConcertInsights/ConcertInsightsVisualizer/Source/ConcertInsightsVisualizer/Private/Util/TimeSyncUtils.cpp b/Engine/Plugins/Developer/Concert/ConcertInsights/ConcertInsightsVisualizer/Source/ConcertInsightsVisualizer/Private/Util/TimeSyncUtils.cpp
new file mode 100644
--- /dev/null
+++ b/Engine/Plugins/Developer/Concert/ConcertInsights/ConcertInsightsVisualizer/Source/ConcertInsightsVisualizer/Private/Util/TimeSyncUtils.cpp
@@ -0,0 +1,45 @@
+// Copyright Epic Games, Inc. All Rights Reserved.
+
+#include "TimeSyncUtils.h"
+
+#include "Misc/DateTime.h"
+#include "Misc/Timespan.h"
+
+namespace UE::ConcertInsightsVisualizer::TimeSyncUtils
+{
+	double ConvertSourceToTargetTime(
+		const FDateTime& TargetInitUtc,
+		const FDateTime& SourceInitUtc,
+		double TargetInitTime,
+		double SourceInitTime,
+		double SourceTime
+		)
+	{
+		/**
+		 * Docu example:
+		 * Global	0123456789_
+		 * Target	[Init567x9]
+		 * Source	-[1Init6y8]
+		 * 
+		 * Given the following input:
+		 * - TargetInitUtc	= 22/02/2024 12:00:01 (random time I choose)
+		 * - SourceInitUtc	= 22/02/2024 12:00:03 (because init event occured 2s after target's init event)
+		 * - TargetInitTime = 1 (see target graph time of I)
+		 * - SourceInitTime = 2 (see source graph time of I)
+		 * - SourceTime		= 7 (see source graph time of y)
+		 */
+		
+		// Going through the above example:
+		// -2 = 22/02/2024 12:00:01 - 22/02/2024 12:00:03
+		const FTimespan SourceToTarget_DeltaTimestamp = TargetInitUtc - SourceInitUtc;
+		
+		// 5 = 7 - 2 - interpretation: 5s have passed since Init event in Source
+		const double SourceTimeRelativeToSourceInit = SourceTime - SourceInitTime;
+		// 7 = 5 - (-2) - interpretation: everything in target happened 2s before anything in the Source timeline
+		const double TargetTimeRelativeToTimeInit = SourceTimeRelativeToSourceInit - SourceToTarget_DeltaTimestamp.GetTotalSeconds();
+		// 8 = 1 + 7 - interpretation: the relative times are anchored relative to the UTC times at which the init events were sent
+		const double TargetTime = TargetInitTime + TargetTimeRelativeToTimeInit;
+		
+		return TargetTime;
+	}
+}
\ No newline at end of file
diff --git a/Engine/Plugins/Developer/Concert/ConcertInsights/ConcertInsightsVisualizer/Source/ConcertInsightsVisualizer/Private/Util/TimeSyncUtils.h b/Engine/Plugins/Developer/Concert/ConcertInsights/ConcertInsightsVisualizer/Source/ConcertInsightsVisualizer/Private/Util/TimeSyncUtils.h
new file mode 100644
--- /dev/null
+++ b/Engine/Plugins/Developer/Concert/ConcertInsights/ConcertInsightsVisualizer/Source/ConcertInsightsVisualizer/Private/Util/TimeSyncUtils.h
@@ -0,0 +1,40 @@
+// Copyright Epic Games, Inc. All Rights Reserved.
+
+#pragma once
+
+struct FDateTime;
+
+namespace UE::ConcertInsightsVisualizer::TimeSyncUtils
+{
+	/**
+	 * Given a source and target timeline, converts SourceTime to the equivalent time in the target timeline.
+	 * Time is assumed to flow at the same pace (i.e. 1s in target timeline corresponds to 1s in source timeline).
+	 *
+	 * Example: Suppose in the below scenario, every character is 1 second and "[" indicates time 0 in the local timeline.
+	 * Global	0123456789_
+	 * Target	[Init567x9]
+	 * Source	-[1Init6y8]
+	 * So this function answers: Given that event y occurs in source's timeline, we want to know the value of x in target's timeline.
+	 * 
+	 * Given the following input:
+	 * - TargetInitUtc	= 22/02/2024 12:00:01 (random time I choose)
+	 * - SourceInitUtc	= 22/02/2024 12:00:03 (because init event occured 2s after target's init event)
+	 * - TargetInitTime = 1 (see target graph time of I)
+	 * - SourceInitTime = 2 (see source graph time of I)
+	 * - SourceTime		= 7 (see source graph time of y)
+	 * This yields x = ConvertSourceToTargetTime(22/02/2024 12:00:02, 22/02/2024 12:00:04, 1, 2, 7) = 8
+	 * 
+	 * @param TargetInitUtc UTC time of target's init event in target timeline
+	 * @param SourceInitUtc UTC time of source's init event in source timeline
+	 * @param TargetInitTime When the target's init event occured in the target timeline (intended in absolute seconds)
+	 * @param SourceInitTime When the source's init event occured in the source timeline (intended in absolute seconds)
+	 * @param SourceTime The time to convert into the target timeline (intended in absolute seconds)
+	 */
+	double ConvertSourceToTargetTime(
+		const FDateTime& TargetInitUtc,
+		const FDateTime& SourceInitUtc,
+		double TargetInitTime,
+		double SourceInitTime,
+		double SourceTime
+		);
+}
diff --git a/Engine/Plugins/Developer/Concert/ConcertInsights/ConcertInsightsVisualizer/Source/ConcertInsightsVisualizer/Public/IConcertInsightsVisualizerModule.h b/Engine/Plugins/Developer/Concert/ConcertInsights/ConcertInsightsVisualizer/Source/ConcertInsightsVisualizer/Public/IConcertInsightsVisualizerModule.h
new file mode 100644
--- /dev/null
+++ b/Engine/Plugins/Developer/Concert/ConcertInsights/ConcertInsightsVisualizer/Source/ConcertInsightsVisualizer/Public/IConcertInsightsVisualizerModule.h
@@ -0,0 +1,38 @@
+// Copyright Epic Games, Inc. All Rights Reserved.
+
+#pragma once
+
+#include "Modules/ModuleManager.h"
+
+namespace UE::ConcertInsightsVisualizer
+{
+	/** This module extends and manages all the Concert UI in Unreal Insights. */
+	class CONCERTINSIGHTSVISUALIZER_API IConcertInsightsVisualizerModule : public IModuleInterface
+	{
+	public:
+
+		/**
+		 * Singleton-like access to this module's interface.  This is just for convenience!
+		 * Beware of calling this during the shutdown phase, though.  Your module might have been unloaded already.
+		 *
+		 * @return Returns singleton instance, loading the module on demand if needed
+		 */
+		static inline IConcertInsightsVisualizerModule& Get()
+		{
+			static const FName ModuleName = "ConcertInsightsVisualizer";
+			return FModuleManager::LoadModuleChecked<IConcertInsightsVisualizerModule>(ModuleName);
+		}
+
+		/**
+		 * Checks to see if this module is loaded and ready.  It is only valid to call Get() during shutdown if IsAvailable() returns true.
+		 *
+		 * @return True if the module is loaded and ready to use
+		 */
+		static inline bool IsAvailable()
+		{
+			static const FName ModuleName = "ConcertInsightsVisualizer";
+			return FModuleManager::Get().IsModuleLoaded(ModuleName);
+		}
+	};
+}
+
diff --git a/Engine/Plugins/Developer/Concert/ConcertInsights/README.txt b/Engine/Plugins/Developer/Concert/ConcertInsights/README.txt
new file mode 100644
--- /dev/null
+++ b/Engine/Plugins/Developer/Concert/ConcertInsights/README.txt
@@ -0,0 +1,47 @@
+========== Summary ==========
+ConcertInsights is a toolchain for tracing steps accross multiple machines in distributed protocols.
+- This adds a status bar entry in the bottom right of the editor to start a synchronized trace (ConcertInsightsCore implements shared logic)
+- Each machine will generate its own .utrace file (ConcertInsightsClient and ConcertInsightsServer listen for synchronized trace requests)
+- ConcertInsightsVisualizer aggregated multiple .utrace files and correlates the machine CPU times
+
+This toolchain is highly experimental. Issues:
+- Time correlation is too imprecise with a fault range of 400 ms. 
+  This is because the FDateTime::UtcNow seems to be accurate only within about 400ms accross machines.
+  To fix this, we should consider a proper time synchronization protocol, e.g. PTP.
+ 
+========== Set-up ==========
+To try this out, the minimum set-up is:
+- In your .uproject file, to Plugins section add 
+	{
+		"Name": "MultiUserClient",
+		"Enabled": true
+	},
+	{
+		"Name": "ConcertInsightsClient",
+		"Enabled": true
+	},
+	{
+		"Name": "ConcertInsightsVisualizer",
+		"Enabled": true
+	}
+- Either unshelf 33391924, or recreate the work in that CL:
+	- Engine/Programs/UnrealMultiUserSlateServerConfig/DefaultEngine.ini: Add +ProgramEnabledPlugins=ConcertInsightsServer
+	- Engine/Programs/UnrealInsights/DefaultEngine.ini: Add +ProgramEnabledPlugins=ConcertInsightsVisualizer
+	
+========== Tutorial with VCam ==========
+To view example tracing, do the following:
+1. Enable Virtual Camera plugin
+2. Add a VCam actor to your level
+3. Launch MU server
+4. Join machine A to MU session
+5. Join machine B to MU session
+6. Connect your mobile device to the VCam on machine A (see https://dev.epicgames.com/community/learning/tutorials/aEeW/unreal-engine-vcam-actor-quick-start-using-pixel-streaming)
+7. On machine A, in the MU session tab, click Replication, then use the Add button to add the VCamActor for replication (the actor should be automatically set up)
+8. On machine B, pilot the VCamActor manually (right-click in outliner > pilot and in viewport toggle piloting on in top-left corner)
+9. Notice that moving your mobile device should now move the viewport on both machine A (source) and B (replicated)
+8. On machine A or B, in the bottom-right, click Multi User > Start synchronized trace
+9. Move your mobile device
+10. On machine A or B, click Multi User > Stop synchronized trace
+11. Open Unreal Insights
+12. Open any of the 3 trace files that were just generated
+13. Disable all channels except for "Concert"
diff --git a/Engine/Plugins/Developer/Concert/ConcertMain/Source/ConcertTransport/Private/ConcertLogGlobal.cpp b/Engine/Plugins/Developer/Concert/ConcertMain/Source/ConcertTransport/Private/ConcertLogGlobal.cpp
--- a/Engine/Plugins/Developer/Concert/ConcertMain/Source/ConcertTransport/Private/ConcertLogGlobal.cpp
+++ b/Engine/Plugins/Developer/Concert/ConcertMain/Source/ConcertTransport/Private/ConcertLogGlobal.cpp
@@ -2,7 +2,5 @@
 
 #include "ConcertLogGlobal.h"
 
-UE_TRACE_CHANNEL_DEFINE(ConcertChannel);
-
 DEFINE_LOG_CATEGORY(LogConcert);
 DEFINE_LOG_CATEGORY(LogConcertDebug);
diff --git a/Engine/Plugins/Developer/Concert/ConcertMain/Source/ConcertTransport/Private/Trace/ConcertProtocolTrace.cpp b/Engine/Plugins/Developer/Concert/ConcertMain/Source/ConcertTransport/Private/Trace/ConcertProtocolTrace.cpp
new file mode 100644
--- /dev/null
+++ b/Engine/Plugins/Developer/Concert/ConcertMain/Source/ConcertTransport/Private/Trace/ConcertProtocolTrace.cpp
@@ -0,0 +1,13 @@
+// Copyright Epic Games, Inc. All Rights Reserved.
+
+#include "Trace/ConcertProtocolTrace.h"
+
+namespace UE::ConcertTrace
+{
+	bool IsTracing(EProtocolSuite Id)
+	{
+		// At the moment it does not matter which protocol is being trace - we just always assume yes.
+		// In the future this might change so keep the unused variable ID for now.
+		return Trace::IsTracing() && Trace::IsChannel(TEXT("ConcertChannel"));
+	}
+}
\ No newline at end of file
diff --git a/Engine/Plugins/Developer/Concert/ConcertMain/Source/ConcertTransport/Private/Trace/ConcertTrace.cpp b/Engine/Plugins/Developer/Concert/ConcertMain/Source/ConcertTransport/Private/Trace/ConcertTrace.cpp
new file mode 100644
--- /dev/null
+++ b/Engine/Plugins/Developer/Concert/ConcertMain/Source/ConcertTransport/Private/Trace/ConcertTrace.cpp
@@ -0,0 +1,5 @@
+// Copyright Epic Games, Inc. All Rights Reserved.
+
+#include "Trace/ConcertTrace.h"
+
+UE_TRACE_CHANNEL_DEFINE(ConcertChannel);
\ No newline at end of file
diff --git a/Engine/Plugins/Developer/Concert/ConcertMain/Source/ConcertTransport/Private/Trace/Private/ConcertProtocolTraceInternal.cpp b/Engine/Plugins/Developer/Concert/ConcertMain/Source/ConcertTransport/Private/Trace/Private/ConcertProtocolTraceInternal.cpp
new file mode 100644
--- /dev/null
+++ b/Engine/Plugins/Developer/Concert/ConcertMain/Source/ConcertTransport/Private/Trace/Private/ConcertProtocolTraceInternal.cpp
@@ -0,0 +1,141 @@
+// Copyright Epic Games, Inc. All Rights Reserved.
+
+#include "Trace/Private/ConcertProtocolTraceInternal.h"
+
+#include "ConcertLogGlobal.h"
+#include "Trace/ConcertTrace.h"
+#include "Trace/Private/InitEventTracking.h"
+
+#include "ProfilingDebugging/TraceAuxiliary.h"
+#include "Trace/Trace.inl"
+
+#if UE_CONCERT_TRACE_ENABLED
+
+UE_TRACE_EVENT_BEGIN(ConcertLogger, Init, NoSync)
+	/** EConcertTraceVersion */
+	UE_TRACE_EVENT_FIELD(uint8, Version)
+	UE_TRACE_EVENT_FIELD(uint32, EndpointId_A)
+	UE_TRACE_EVENT_FIELD(uint32, EndpointId_B)
+	UE_TRACE_EVENT_FIELD(uint32, EndpointId_C)
+	UE_TRACE_EVENT_FIELD(uint32, EndpointId_D)
+	/** Time the trace was started. Used to correlate times relative to other machines. */
+	UE_TRACE_EVENT_FIELD(int64, DateTimeTicks)
+	/** Cycle at which this init was generated. Used to correlate times relative to other machines. */
+	UE_TRACE_EVENT_FIELD(uint64, Cycle)
+	UE_TRACE_EVENT_FIELD(UE::Trace::WideString, ClientDisplayName)
+	UE_TRACE_EVENT_FIELD(bool, IsServer)
+	UE_TRACE_EVENT_FIELD(bool, HasEndpointId)
+	UE_TRACE_EVENT_FIELD(bool, HasDisplayName)
+UE_TRACE_EVENT_END()
+
+UE_TRACE_EVENT_BEGIN(ConcertLogger, ObjectTransmissionStart, NoSync)
+	UE_TRACE_EVENT_FIELD(uint32, DestEndpointId_A)
+	UE_TRACE_EVENT_FIELD(uint32, DestEndpointId_B)
+	UE_TRACE_EVENT_FIELD(uint32, DestEndpointId_C)
+	UE_TRACE_EVENT_FIELD(uint32, DestEndpointId_D)
+	UE_TRACE_EVENT_FIELD(uint8, Protocol)
+	UE_TRACE_EVENT_FIELD(uint32, SequenceId)
+	UE_TRACE_EVENT_FIELD(UE::Trace::WideString, ObjectPath)
+	UE_TRACE_EVENT_FIELD(uint64, Cycle)
+UE_TRACE_EVENT_END()
+
+UE_TRACE_EVENT_BEGIN(ConcertLogger, ObjectTransmissionReceive, NoSync)
+	UE_TRACE_EVENT_FIELD(uint8, Protocol)
+	UE_TRACE_EVENT_FIELD(uint32, SequenceId)
+	UE_TRACE_EVENT_FIELD(UE::Trace::WideString, ObjectPath)
+	UE_TRACE_EVENT_FIELD(uint64, Cycle)
+UE_TRACE_EVENT_END()
+
+UE_TRACE_EVENT_BEGIN(ConcertLogger, ObjectSink, NoSync)
+	UE_TRACE_EVENT_FIELD(uint8, Protocol)
+	UE_TRACE_EVENT_FIELD(uint32, SequenceId)
+	UE_TRACE_EVENT_FIELD(UE::Trace::AnsiString, SinkName)
+	UE_TRACE_EVENT_FIELD(UE::Trace::WideString, ObjectPath)
+	UE_TRACE_EVENT_FIELD(uint64, Cycle)
+UE_TRACE_EVENT_END()
+
+namespace UE::ConcertTrace
+{
+	void Init(const TOptional<FGuid>& LocalEndpointId, const FString& ClientDisplayName, bool bIsServer)
+	{
+		const bool bIsConnected = FTraceAuxiliary::IsConnected();
+		const bool bHasSentInitEventBefore = HasSentInitEventToCurrentSession();
+		if (!bIsConnected
+			|| bHasSentInitEventBefore)
+		{
+			return;
+		}
+		
+		const bool bIsChannelEnabled = UE_TRACE_CHANNELEXPR_IS_ENABLED(ConcertChannel);
+		if (!bIsChannelEnabled)
+		{
+			// Warn so developer can debug it
+			UE_LOG(LogConcert, Warning, TEXT("Skipping Concert init event because the Concert channel is not enabled."));
+			return;
+		}
+		
+		OnSendInitEvent();
+		
+		UE_TRACE_LOG(ConcertLogger, Init, ConcertChannel)
+		<< Init.Version(static_cast<uint8>(EConcertTraceVersion::Latest))
+		<< Init.EndpointId_A(LocalEndpointId.IsSet() ? LocalEndpointId->A : 0) 
+		<< Init.EndpointId_B(LocalEndpointId.IsSet() ? LocalEndpointId->B : 0) 
+		<< Init.EndpointId_C(LocalEndpointId.IsSet() ? LocalEndpointId->C : 0) 
+		<< Init.EndpointId_D(LocalEndpointId.IsSet() ? LocalEndpointId->D : 0)
+		<< Init.DateTimeTicks(FDateTime::UtcNow().GetTicks())
+		<< Init.Cycle(FPlatformTime::Cycles64())
+		<< Init.ClientDisplayName(*ClientDisplayName)
+		<< Init.IsServer(bIsServer) 
+		<< Init.HasEndpointId(LocalEndpointId.IsSet())
+		<< Init.HasDisplayName(!ClientDisplayName.IsEmpty());
+	}
+	
+	void LogTransmissionStart(uint8 Protocol, uint32 SequenceId, const FSoftObjectPath& ObjectPath, const FGuid& TargetEndpointId)
+	{
+		if (!ShouldTraceConcertProtocols())
+		{
+			return;
+		}
+		
+		UE_TRACE_LOG(ConcertLogger, ObjectTransmissionStart, ConcertChannel)
+		<< ObjectTransmissionStart.DestEndpointId_A(TargetEndpointId.A) 
+		<< ObjectTransmissionStart.DestEndpointId_B(TargetEndpointId.B) 
+		<< ObjectTransmissionStart.DestEndpointId_C(TargetEndpointId.C) 
+		<< ObjectTransmissionStart.DestEndpointId_D(TargetEndpointId.D) 
+		<< ObjectTransmissionStart.Protocol(Protocol) 
+		<< ObjectTransmissionStart.SequenceId(SequenceId) 
+		<< ObjectTransmissionStart.ObjectPath(*ObjectPath.ToString())
+		<< ObjectTransmissionStart.Cycle(FPlatformTime::Cycles64());
+	}
+	
+	void LogTransmissionReceive(uint8 Protocol, uint32 SequenceId, const FSoftObjectPath& ObjectPath)
+	{
+		if (!ShouldTraceConcertProtocols())
+		{
+			return;
+		}
+		
+		UE_TRACE_LOG(ConcertLogger, ObjectTransmissionReceive, ConcertChannel) 
+		<< ObjectTransmissionReceive.Protocol(Protocol) 
+		<< ObjectTransmissionReceive.SequenceId(SequenceId) 
+		<< ObjectTransmissionReceive.ObjectPath(*ObjectPath.ToString())
+		<< ObjectTransmissionReceive.Cycle(FPlatformTime::Cycles64());
+	}
+	
+	void LogSink(uint8 Protocol, uint32 SequenceId, const ANSICHAR* SinkName, const FSoftObjectPath& ObjectPath)
+	{
+		if (!ShouldTraceConcertProtocols())
+		{
+			return;
+		}
+		
+		UE_TRACE_LOG(ConcertLogger, ObjectSink, ConcertChannel) 
+		<< ObjectSink.Protocol(Protocol) 
+		<< ObjectSink.SequenceId(SequenceId) 
+		<< ObjectSink.SinkName(SinkName) 
+		<< ObjectSink.ObjectPath(*ObjectPath.ToString())
+		<< ObjectSink.Cycle(FPlatformTime::Cycles64());
+	}
+}
+
+#endif
\ No newline at end of file
diff --git a/Engine/Plugins/Developer/Concert/ConcertMain/Source/ConcertTransport/Private/Trace/Private/ConcertScopedObjectTrace.cpp b/Engine/Plugins/Developer/Concert/ConcertMain/Source/ConcertTransport/Private/Trace/Private/ConcertScopedObjectTrace.cpp
new file mode 100644
--- /dev/null
+++ b/Engine/Plugins/Developer/Concert/ConcertMain/Source/ConcertTransport/Private/Trace/Private/ConcertScopedObjectTrace.cpp
@@ -0,0 +1,74 @@
+// Copyright Epic Games, Inc. All Rights Reserved.
+
+#include "Trace/Private/ConcertScopedObjectTrace.h"
+
+#include "Trace/ConcertTrace.h"
+
+#if UE_CONCERT_TRACE_ENABLED
+
+#include "Trace/Private/InitEventTracking.h"
+
+#include "HAL/PlatformTime.h"
+#include "Trace/Trace.inl"
+
+UE_TRACE_EVENT_BEGIN(ConcertLogger, ObjectTraceBegin, NoSync)
+	UE_TRACE_EVENT_FIELD(uint8, Protocol)
+	UE_TRACE_EVENT_FIELD(uint32, SequenceId)
+	UE_TRACE_EVENT_FIELD(UE::Trace::WideString, ObjectPath)
+	UE_TRACE_EVENT_FIELD(uint64, Cycle)
+	UE_TRACE_EVENT_FIELD(UE::Trace::AnsiString, EventName)
+UE_TRACE_EVENT_END()
+
+UE_TRACE_EVENT_BEGIN(ConcertLogger, ObjectTraceEnd, NoSync)
+	UE_TRACE_EVENT_FIELD(uint8, Protocol)
+	UE_TRACE_EVENT_FIELD(uint32, SequenceId)
+	UE_TRACE_EVENT_FIELD(UE::Trace::WideString, ObjectPath)
+	UE_TRACE_EVENT_FIELD(uint64, Cycle)
+	UE_TRACE_EVENT_FIELD(UE::Trace::AnsiString, EventName)
+UE_TRACE_EVENT_END()
+
+namespace UE::ConcertTrace
+{
+	FConcertScopedObjectTrace::FConcertScopedObjectTrace(
+		uint8 InProtocol,
+		uint32 InSequenceId,
+		const ANSICHAR* InEventName,
+		FSoftObjectPath InObjectPath
+		)
+		: bShouldTrace(ShouldTraceConcertProtocols())
+		, Protocol(InProtocol)
+		, SequenceId(InSequenceId)
+		, EventName(InEventName)
+		, ObjectPath(MoveTemp(InObjectPath))
+	{
+		if (!bShouldTrace)
+		{
+			return;
+		}
+		
+		// Send one event in constructor with begin time and one in destructor with end time instead of one event in destructor that has begin and end time.
+		// Reason: Needed to for TMonotonicTimeline::AppendBeginEvent API used internally by FProtocolMultiEndpointProvider for nested CPU events. 
+		UE_TRACE_LOG(ConcertLogger, ObjectTraceBegin, ConcertChannel)
+		<< ObjectTraceBegin.Protocol(Protocol) 
+		<< ObjectTraceBegin.SequenceId(SequenceId) 
+		<< ObjectTraceBegin.ObjectPath(*ObjectPath.ToString())
+		<< ObjectTraceBegin.Cycle(FPlatformTime::Cycles64())
+		<< ObjectTraceBegin.EventName(EventName);
+	}
+
+	FConcertScopedObjectTrace::~FConcertScopedObjectTrace()
+	{
+		if (!bShouldTrace)
+		{
+			return;
+		}
+		
+		UE_TRACE_LOG(ConcertLogger, ObjectTraceEnd, ConcertChannel)
+		<< ObjectTraceEnd.Protocol(Protocol) 
+		<< ObjectTraceEnd.SequenceId(SequenceId) 
+		<< ObjectTraceEnd.ObjectPath(*ObjectPath.ToString())
+		<< ObjectTraceEnd.Cycle(FPlatformTime::Cycles64())
+		<< ObjectTraceEnd.EventName(EventName);
+	}
+}
+#endif
\ No newline at end of file
diff --git a/Engine/Plugins/Developer/Concert/ConcertMain/Source/ConcertTransport/Private/Trace/Private/InitEventTracking.cpp b/Engine/Plugins/Developer/Concert/ConcertMain/Source/ConcertTransport/Private/Trace/Private/InitEventTracking.cpp
new file mode 100644
--- /dev/null
+++ b/Engine/Plugins/Developer/Concert/ConcertMain/Source/ConcertTransport/Private/Trace/Private/InitEventTracking.cpp
@@ -0,0 +1,38 @@
+// Copyright Epic Games, Inc. All Rights Reserved.
+
+#include "InitEventTracking.h"
+
+#include "Trace/ConcertTraceConfig.h"
+
+#if UE_CONCERT_TRACE_ENABLED
+
+#include "ConcertLogGlobal.h"
+#include "Trace/ConcertTrace.h"
+
+#include "ProfilingDebugging/TraceAuxiliary.h"
+#include "Trace/Trace.h"
+
+namespace UE::ConcertTrace
+{
+	/** Trace destinations that the local application instance has already sent the init event to. The init event should only appear in every .utrace file once. */
+	TSet<FString> TraceDestinationsWithInitEvent;
+
+	bool HasSentInitEventToCurrentSession()
+	{
+		return TraceDestinationsWithInitEvent.Contains(FTraceAuxiliary::GetTraceDestinationString());
+	}
+	
+	bool ShouldTraceConcertProtocols()
+	{
+		return UE_TRACE_CHANNELEXPR_IS_ENABLED(ConcertChannel)
+			&& HasSentInitEventToCurrentSession();
+	}
+
+	void OnSendInitEvent()
+	{
+		UE_LOG(LogConcert, Log, TEXT("Sending trace init event"));
+		TraceDestinationsWithInitEvent.Add(FTraceAuxiliary::GetTraceDestinationString());
+	}
+}
+
+#endif
\ No newline at end of file
diff --git a/Engine/Plugins/Developer/Concert/ConcertMain/Source/ConcertTransport/Private/Trace/Private/InitEventTracking.h b/Engine/Plugins/Developer/Concert/ConcertMain/Source/ConcertTransport/Private/Trace/Private/InitEventTracking.h
new file mode 100644
--- /dev/null
+++ b/Engine/Plugins/Developer/Concert/ConcertMain/Source/ConcertTransport/Private/Trace/Private/InitEventTracking.h
@@ -0,0 +1,21 @@
+// Copyright Epic Games, Inc. All Rights Reserved.
+
+#pragma once
+
+#include "Trace/ConcertTraceConfig.h"
+
+#if UE_CONCERT_TRACE_ENABLED
+
+namespace UE::ConcertTrace
+{
+	/** Checks whether an init event has been sent for the current trace session already. */
+	bool HasSentInitEventToCurrentSession();
+
+	/** @return Whether protocol tracing should generate data. */
+	bool ShouldTraceConcertProtocols();
+
+	/** Tracks that for the current tracing session, and init event was been sent. */
+	void OnSendInitEvent();
+}
+
+#endif
diff --git a/Engine/Plugins/Developer/Concert/ConcertMain/Source/ConcertTransport/Public/ConcertLogGlobal.h b/Engine/Plugins/Developer/Concert/ConcertMain/Source/ConcertTransport/Public/ConcertLogGlobal.h
--- a/Engine/Plugins/Developer/Concert/ConcertMain/Source/ConcertTransport/Public/ConcertLogGlobal.h
+++ b/Engine/Plugins/Developer/Concert/ConcertMain/Source/ConcertTransport/Public/ConcertLogGlobal.h
@@ -3,14 +3,8 @@
 #pragma once
 
 #include "Logging/LogMacros.h"
-#include "ProfilingDebugging/CpuProfilerTrace.h"
-
-#if UE_TRACE_ENABLED && !UE_BUILD_SHIPPING
-UE_TRACE_CHANNEL_EXTERN(ConcertChannel, CONCERTTRANSPORT_API);
-#define SCOPED_CONCERT_TRACE(TraceName) TRACE_CPUPROFILER_EVENT_SCOPE_ON_CHANNEL(TraceName, ConcertChannel)
-#else
-#define SCOPED_CONCERT_TRACE(TraceName)
-#endif
+// This include is for backwards compatibility. SCOPED_CONCERT_TRACE and ConcertChannel used to be defined here. TODO DP: Deprecate this correctly
+#include "Trace/ConcertTrace.h" 
 
 // The default debug verbosity level can be set as a GlobalDefinitions in an executable Target.cs build file to adapt to this executable context.
 #ifndef UE_LOG_CONCERT_DEBUG_VERBOSITY_LEVEL
diff --git a/Engine/Plugins/Developer/Concert/ConcertMain/Source/ConcertTransport/Public/Trace/ConcertProtocolTrace.h b/Engine/Plugins/Developer/Concert/ConcertMain/Source/ConcertTransport/Public/Trace/ConcertProtocolTrace.h
new file mode 100644
--- /dev/null
+++ b/Engine/Plugins/Developer/Concert/ConcertMain/Source/ConcertTransport/Public/Trace/ConcertProtocolTrace.h
@@ -0,0 +1,126 @@
+// Copyright Epic Games, Inc. All Rights Reserved.
+
+#pragma once
+
+#include "Trace/ConcertTraceConfig.h"
+#include "Trace/Private/ConcertProtocolTraceInternal.h"
+
+/*
+ * This file defines the Concert Protocol Tracing macros.
+ * Unreal Insights does not understand it by default and we add extensions in the ConcertInsights plugin to visualize it.
+ * 
+ * The heart of Concert Protocol Tracing are sequences.
+ * A sequence consists of actions that are performed by multiple endpoints.
+ * Example of a sequence:
+ *  - Every frame, client 1 collects data to replicate and sends it to the server.
+ *  - The server enqueues the data and sends it to client 2.
+ *  - Client 2 applies the data.
+ * Sequences are associated with a EProtocolSuite and grouped accordingly by the ConcertInsights extension.
+ *
+ * For now there is only one type of trace: traced objects. See CONCERT_TRACE_OBJECT_ Macros.
+ * These associate actions with an object and a sequence ID, which groups together related changes.
+ * For example, in the context of replication a single frame's data being sent across the clients gets a single sequence ID:
+ *	- client 1 could trace a. how long it takes to serialize the data and then b. how long to compress it.
+ *	- the server could trace how long it takes to process the data
+ *	- client 2 could trace a. how long it apply the data to the UObject
+ *	
+ * ConcertInsights visualizes networking transport times.
+ * Transmission starts with CONCERT_TRACE_OBJECT_TRANSMISSION_START and ends with CONCERT_TRACE_OBJECT_TRANSMISSION_RECEIVE.
+ * 
+ * TODO DP: We might want to add CONCERT_TRACE_OBJECT_SEND and CONCERT_TRACE_OBJECT_RECEIVE macros to be explicit about sending and receiving.
+ * 
+ * Every endpoint writes to its own .utrace file and must be in context of a session ID (i.e. clients are in a session, server is told which session is being traced).
+ * When beginning to trace, the first event sent to the .utrace files is CONCERT_TRACE_INIT, which sets data needed to aggregate the .utrace files, such as endpoint IDs.
+ * Finally, ConcertInsights aggregates the .utrace files in the UI. The related .utrace files are grouped by SessionID, which Unreal Insights saves
+ * in the file by reading the -tracesessionguid command line; CONCERT_TRACE_INIT associates the .utrace file with the client / server endpoint ID.
+ * ConcertInsightsSynchronizedTrace implements synchronized recording.
+ */
+
+namespace UE::ConcertTrace
+{
+	enum class EProtocolSuite : uint8
+	{
+		Unknown,
+
+		/** The traces are related to the replication system */
+		Replication
+	};
+
+	FORCEINLINE uint8 ProtocolSuiteToInt(EProtocolSuite Id) { return static_cast<uint8>(Id); }
+
+	/** @return Whether the given protocol is currently being traced */
+	CONCERTTRANSPORT_API bool IsTracing(EProtocolSuite Id);
+	/** @return Whether replication is currently being traced. */
+	FORCEINLINE bool IsTracingReplication() { return IsTracing(EProtocolSuite::Replication); }
+}
+
+#if UE_CONCERT_TRACE_ENABLED
+
+/**
+ * Event that is emitted when a trace is started in the local instance.
+ * 
+ * @param InEndpointIdAsTOptional TOptional<FGuid> that identifies the local instance endpoint ID in the active session.
+ *	This should always be set if started via ConcertInsightsSynchronizedTrace.
+ *	May be unset when the trace is conventionally started by user.
+ * @param IsServer Whether the local instance is the server
+ */
+#define CONCERT_TRACE_INIT(EndpointIdAsTOptional, InEndpointDisplayFString, IsServer) INTERNAL_CONCERT_TRACE_INIT(EndpointIdAsTOptional, InEndpointDisplayFString, IsServer)
+
+/********** Generics Macros - Create new macros that wrap these when you add a new traced protocol **********/
+
+/**
+ * Logs an event with a start an end time.
+ * In Insights, this will show a track that begins at the current time and ends when the scope terminates.
+ * Supports nesting.
+ * 
+ * @param EventName The name of the event to show
+ * @param ObjectPath The object this event pertains to. Events with the same ObjectPath and SequenceID are grouped in one track.
+ * @param SequenceId Which change to associate this event to (use-case defined). Events with the same ObjectPath and SequenceID are grouped in one track.
+ * @param ProtocolId The protocol to group this event under.
+ */
+#define CONCERT_TRACE_OBJECT_SCOPE(EventName, ObjectPath, SequenceId, ProtocolId) INTERNAL_CONCERT_TRACE_OBJECT_SCOPE(EventName, ObjectPath, SequenceId, ProtocolId)
+
+/** Traces that the local endpoint has handed the object's sequence data to message bus. The time from now until CONCERT_TRACE_OBJECT_TRANSMISSION_RECEIVE is visualized as transit time. */
+#define CONCERT_TRACE_OBJECT_TRANSMISSION_START(DestinationEndpointId, ObjectPath, SequenceId, ProtocolId) INTERNAL_CONCERT_TRACE_OBJECT_TRANSMISSION_START(DestinationEndpointId, ObjectPath, SequenceId, ProtocolId)
+/** Traces that the local endpoint has received the object's sequence data. The time from now the previous CONCERT_TRACE_OBJECT_TRANSMISSION_RECEIVE until now is visualized as transit time. */
+#define CONCERT_TRACE_OBJECT_TRANSMISSION_RECEIVE(ObjectPath, SequenceId, ProtocolId) INTERNAL_CONCERT_TRACE_OBJECT_TRANSMISSION_RECEIVE(ObjectPath, SequenceId, ProtocolId)
+/** The object's data has been fully processed. End of the object track. */
+#define CONCERT_TRACE_OBJECT_SINK(SinkName, ObjectPath, SequenceId, ProtocolId) INTERNAL_CONCERT_TRACE_OBJECT_SINK(SinkName, ObjectPath, SequenceId, ProtocolId)
+
+/********** Replication Protocol **********/
+
+// TODO DP: Add CONCERT_TRACE_REPLICATION_OBJECT_ANONYMOUS_CONTEXT_SCOPE(EventName) which could be used to give more context of what other actions are being performed.
+// Example: This could be added at the root of Tick() so we know what other parts of Concert are delaying packet delivery.
+// Alternative could also just be using CONCERT_TRACE_SCOPE ...
+
+/**
+ * Logs an event with a start an end time.
+ * In Insights, this will show a track that begins at the current time and ends at the start time of the next event.
+ * Supports nesting.
+ * 
+ * @param EventName The name of the event to show
+ * @param ObjectPath The object this event pertains to.
+ * @param SequenceId Determined by sending client and passed along to all other endpoints. 
+ */
+#define CONCERT_TRACE_REPLICATION_OBJECT_SCOPE(EventName, ObjectPath, SequenceId) CONCERT_TRACE_OBJECT_SCOPE(EventName, ObjectPath, SequenceId, ProtocolSuiteToInt(UE::ConcertTrace::EProtocolSuite::Replication))
+
+/** Traces that the local endpoint has handed the object's sequence data to message bus. The time from now until CONCERT_TRACE_REPLICATION_OBJECT_TRANSMISSION_RECEIVE is visualized as transit time. */
+#define CONCERT_TRACE_REPLICATION_OBJECT_TRANSMISSION_START(DestinationEndpointId, ObjectPath, SequenceId) CONCERT_TRACE_OBJECT_TRANSMISSION_START(DestinationEndpointId, ObjectPath, SequenceId, ProtocolSuiteToInt(UE::ConcertTrace::EProtocolSuite::Replication))
+/** Traces that the local endpoint has received the object's sequence data. The time from now the previous CONCERT_TRACE_REPLICATION_OBJECT_TRANSMISSION_START until now is visualized as transit time. */
+#define CONCERT_TRACE_REPLICATION_OBJECT_TRANSMISSION_RECEIVE(ObjectPath, SequenceId) CONCERT_TRACE_OBJECT_TRANSMISSION_RECEIVE(ObjectPath, SequenceId, ProtocolSuiteToInt(UE::ConcertTrace::EProtocolSuite::Replication))
+/** The object's data has been fully processed. End of the object track. */
+#define CONCERT_TRACE_REPLICATION_OBJECT_SINK(SinkName, ObjectPath, SequenceId) CONCERT_TRACE_OBJECT_SINK(SinkName, ObjectPath, SequenceId, ProtocolSuiteToInt(UE::ConcertTrace::EProtocolSuite::Replication))
+
+#else
+#define CONCERT_TRACE_INIT(SessionId)
+#define CONCERT_TRACE_OBJECT_SCOPE(EventName, ObjectPath, SequenceId, ProtocolId)
+#define CONCERT_TRACE_OBJECT_EVENT(EventName, ObjectPath, SequenceId, ProtocolId)
+#define CONCERT_TRACE_REPLICATION_OBJECT_SCOPE(EventName, ObjectPath, SequenceId)
+#define CONCERT_TRACE_REPLICATION_OBJECT_EVENT(EventName, ObjectPath, SequenceId)
+#define CONCERT_TRACE_OBJECT_TRANSMISSION_START(DestinationEndpointId, ObjectPath, SequenceId, ProtocolId)
+#define CONCERT_TRACE_OBJECT_TRANSMISSION_RECEIVE(ObjectPath, SequenceId, ProtocolId)
+#define CONCERT_TRACE_OBJECT_SINK(ObjectPath, SequenceId, ProtocolId)
+#define CONCERT_TRACE_REPLICATION_OBJECT_TRANSMISSION_START(DestinationEndpointId, ObjectPath, SequenceId)
+#define CONCERT_TRACE_REPLICATION_OBJECT_TRANSMISSION_RECEIVE(ObjectPath, SequenceId)
+#define CONCERT_TRACE_REPLICATION_OBJECT_SINK(SinkName, ObjectPath, SequenceId)
+#endif
\ No newline at end of file
diff --git a/Engine/Plugins/Developer/Concert/ConcertMain/Source/ConcertTransport/Public/Trace/ConcertTrace.h b/Engine/Plugins/Developer/Concert/ConcertMain/Source/ConcertTransport/Public/Trace/ConcertTrace.h
new file mode 100644
--- /dev/null
+++ b/Engine/Plugins/Developer/Concert/ConcertMain/Source/ConcertTransport/Public/Trace/ConcertTrace.h
@@ -0,0 +1,18 @@
+// Copyright Epic Games, Inc. All Rights Reserved.
+
+#pragma once
+
+#include "ConcertTraceConfig.h"
+#include "Trace/Trace.h"
+
+/*
+ * This file defines the common UE tracing utilities.
+ * Common means the traces that Unreal Insights already understands without any extensions added by ConcertInsights plugin.
+ */
+
+#if UE_CONCERT_TRACE_ENABLED
+	UE_TRACE_CHANNEL_EXTERN(ConcertChannel, CONCERTTRANSPORT_API);
+#	define SCOPED_CONCERT_TRACE(TraceName) TRACE_CPUPROFILER_EVENT_SCOPE_ON_CHANNEL(TraceName, ConcertChannel)
+#else
+#	define SCOPED_CONCERT_TRACE(TraceName)
+#endif
\ No newline at end of file
diff --git a/Engine/Plugins/Developer/Concert/ConcertMain/Source/ConcertTransport/Public/Trace/ConcertTraceConfig.h b/Engine/Plugins/Developer/Concert/ConcertMain/Source/ConcertTransport/Public/Trace/ConcertTraceConfig.h
new file mode 100644
--- /dev/null
+++ b/Engine/Plugins/Developer/Concert/ConcertMain/Source/ConcertTransport/Public/Trace/ConcertTraceConfig.h
@@ -0,0 +1,14 @@
+// Copyright Epic Games, Inc. All Rights Reserved.
+
+#pragma once
+
+#include "ProfilingDebugging/CpuProfilerTrace.h"
+#include "Trace/Config.h"
+
+#if !defined(UE_CONCERT_TRACE_ENABLED)
+#	if UE_TRACE_ENABLED && CPUPROFILERTRACE_ENABLED && !UE_BUILD_SHIPPING
+#		define UE_CONCERT_TRACE_ENABLED 1
+#	else
+#		define UE_CONCERT_TRACE_ENABLED 0
+#	endif
+#endif
\ No newline at end of file
diff --git a/Engine/Plugins/Developer/Concert/ConcertMain/Source/ConcertTransport/Public/Trace/Private/ConcertProtocolTraceInternal.h b/Engine/Plugins/Developer/Concert/ConcertMain/Source/ConcertTransport/Public/Trace/Private/ConcertProtocolTraceInternal.h
new file mode 100644
--- /dev/null
+++ b/Engine/Plugins/Developer/Concert/ConcertMain/Source/ConcertTransport/Public/Trace/Private/ConcertProtocolTraceInternal.h
@@ -0,0 +1,38 @@
+// Copyright Epic Games, Inc. All Rights Reserved.
+
+#pragma once
+
+#include "Trace/ConcertTraceConfig.h"
+
+#if UE_CONCERT_TRACE_ENABLED
+
+#include "Trace/Private/ConcertScopedObjectTrace.h"
+
+namespace UE::ConcertTrace
+{
+	/** Sent as part of the Init event for backwards compatibility. */
+	enum class EConcertTraceVersion : uint8
+	{
+		/** Initial version of Concert Insights. */
+		Initial = 0,
+
+		// Add new entries above
+		LatestPlusOne,
+		Latest = LatestPlusOne - 1
+	};
+	
+	CONCERTTRANSPORT_API void Init(const TOptional<FGuid>& LocalEndpointId, const FString& ClientDisplayName, bool bIsServer);
+	CONCERTTRANSPORT_API void LogTransmissionStart(uint8 Protocol, uint32 SequenceId, const FSoftObjectPath& ObjectPath, const FGuid& TargetEndpointId);
+	CONCERTTRANSPORT_API void LogTransmissionReceive(uint8 Protocol, uint32 SequenceId, const FSoftObjectPath& ObjectPath);
+	CONCERTTRANSPORT_API void LogSink(uint8 Protocol, uint32 SequenceId, const ANSICHAR* SinkName, const FSoftObjectPath& ObjectPath);
+}
+
+#define INTERNAL_CONCERT_TRACE_INIT(InEndpointIdTOptional, InEndpointDisplayFString, InIsServer) UE::ConcertTrace::Init(InEndpointIdTOptional, InEndpointDisplayFString, InIsServer);
+
+#define INTERNAL_CONCERT_TRACE_OBJECT_SCOPE(InEventName, InObjectPath, InSequenceId, InProtocolId) UE::ConcertTrace::FConcertScopedObjectTrace PREPROCESSOR_JOIN(TheScope, __LINE__)(InProtocolId, InSequenceId, #InEventName, InObjectPath);
+
+#define INTERNAL_CONCERT_TRACE_OBJECT_TRANSMISSION_START(InDestinationEndpointId, InObjectPath, InSequenceId, InProtocolId) UE::ConcertTrace::LogTransmissionStart(InProtocolId, InSequenceId, InObjectPath, InDestinationEndpointId);
+#define INTERNAL_CONCERT_TRACE_OBJECT_TRANSMISSION_RECEIVE(InObjectPath, InSequenceId, InProtocolId) UE::ConcertTrace::LogTransmissionReceive(InProtocolId, InSequenceId, InObjectPath);
+#define INTERNAL_CONCERT_TRACE_OBJECT_SINK(InSinkName, InObjectPath, InSequenceId, InProtocolId) UE::ConcertTrace::LogSink(InProtocolId, InSequenceId, #InSinkName, InObjectPath);
+
+#endif
\ No newline at end of file
diff --git a/Engine/Plugins/Developer/Concert/ConcertMain/Source/ConcertTransport/Public/Trace/Private/ConcertReplicationSequenceContext.h b/Engine/Plugins/Developer/Concert/ConcertMain/Source/ConcertTransport/Public/Trace/Private/ConcertReplicationSequenceContext.h
new file mode 100644
--- /dev/null
+++ b/Engine/Plugins/Developer/Concert/ConcertMain/Source/ConcertTransport/Public/Trace/Private/ConcertReplicationSequenceContext.h
@@ -0,0 +1,14 @@
+// Copyright Epic Games, Inc. All Rights Reserved.
+
+#pragma once
+
+namespace UE::ConcertTrace::Private
+{
+	struct FConcertReplicationSequenceContext
+	{
+		FSoftObjectPath Object;
+		int32 FrameNumber;
+
+		FString ToString() const { return Object.ToString() + FString::FromInt(FrameNumber); }
+	};
+}
\ No newline at end of file
diff --git a/Engine/Plugins/Developer/Concert/ConcertMain/Source/ConcertTransport/Public/Trace/Private/ConcertScopedObjectTrace.h b/Engine/Plugins/Developer/Concert/ConcertMain/Source/ConcertTransport/Public/Trace/Private/ConcertScopedObjectTrace.h
new file mode 100644
--- /dev/null
+++ b/Engine/Plugins/Developer/Concert/ConcertMain/Source/ConcertTransport/Public/Trace/Private/ConcertScopedObjectTrace.h
@@ -0,0 +1,30 @@
+// Copyright Epic Games, Inc. All Rights Reserved.
+
+#pragma once
+
+#include "Trace/ConcertTraceConfig.h"
+
+#if UE_CONCERT_TRACE_ENABLED
+
+#include "HAL/Platform.h"
+#include "Templates/UnrealTemplate.h"
+#include "UObject/SoftObjectPath.h"
+
+namespace UE::ConcertTrace
+{
+	class CONCERTTRANSPORT_API FConcertScopedObjectTrace : public FNoncopyable
+	{
+		const bool bShouldTrace;
+		
+		const uint8 Protocol;
+		const uint32 SequenceId;
+		const ANSICHAR* EventName;
+		const FSoftObjectPath ObjectPath;
+		
+	public:
+
+		FConcertScopedObjectTrace(uint8 InProtocol, uint32 InSequenceId, const ANSICHAR* InEventName, FSoftObjectPath InObjectPath);
+		~FConcertScopedObjectTrace();
+	};
+}
+#endif
\ No newline at end of file
diff --git a/Engine/Plugins/Developer/Concert/ConcertSync/ConcertSyncClient/Source/ConcertSyncClient/ConcertSyncClient.Build.cs b/Engine/Plugins/Developer/Concert/ConcertSync/ConcertSyncClient/Source/ConcertSyncClient/ConcertSyncClient.Build.cs
--- a/Engine/Plugins/Developer/Concert/ConcertSync/ConcertSyncClient/Source/ConcertSyncClient/ConcertSyncClient.Build.cs
+++ b/Engine/Plugins/Developer/Concert/ConcertSync/ConcertSyncClient/Source/ConcertSyncClient/ConcertSyncClient.Build.cs
@@ -34,6 +34,7 @@ namespace UnrealBuildTool.Rules
 					"Slate",
 					"SourceControl",
 					"Serialization",
+					"TraceLog"
                 }
             );
 
diff --git a/Engine/Plugins/Developer/Concert/ConcertSync/ConcertSyncClient/Source/ConcertSyncClient/Private/Replication/Processing/ClientReplicationDataCollector.cpp b/Engine/Plugins/Developer/Concert/ConcertSync/ConcertSyncClient/Source/ConcertSyncClient/Private/Replication/Processing/ClientReplicationDataCollector.cpp
--- a/Engine/Plugins/Developer/Concert/ConcertSync/ConcertSyncClient/Source/ConcertSyncClient/Private/Replication/Processing/ClientReplicationDataCollector.cpp
+++ b/Engine/Plugins/Developer/Concert/ConcertSync/ConcertSyncClient/Source/ConcertSyncClient/Private/Replication/Processing/ClientReplicationDataCollector.cpp
@@ -158,7 +158,7 @@ namespace UE::ConcertSyncClient::Replication
 		}
 	}
 
-	void FClientReplicationDataCollector::ForEachPendingObject(TFunctionRef<void(const FConcertReplicatedObjectId&)> ProcessItemFunc) const
+	void FClientReplicationDataCollector::ForEachPendingObject(TFunctionRef<void(const ConcertSyncCore::FPendingObjectReplicationInfo&)> ProcessItemFunc) const
 	{
 		for (const TPair<FSoftObjectPath, TArray<FObjectInfo>>& Pair : ObjectsToReplicate)
 		{
@@ -182,7 +182,12 @@ namespace UE::ConcertSyncClient::Replication
 				const FSoftObjectPath ObjectPath = Object;
 				if (Object && ensureMsgf(ObjectPath == Object, TEXT("Sanity check: the bridge gave us an object with a different path!")))
 				{
-					ProcessItemFunc(ObjectId);
+					ConcertSyncCore::FPendingObjectReplicationInfo Info
+					{
+						FConcertReplicatedObjectId {FConcertObjectInStreamID{ ObjectInfo.StreamId, Object },ClientId },
+						ObjectInfo.ReplicationSequenceId
+					};
+					ProcessItemFunc(Info);
 				}
 			}
 		}
@@ -194,7 +199,7 @@ namespace UE::ConcertSyncClient::Replication
 		TFunctionRef<void(FConcertSessionSerializedPayload&& Payload)> ProcessMoveable
 		)
 	{
-		const TArray<FObjectInfo>* ObjectInfos = ObjectsToReplicate.Find(ObjectToProcess.Object);
+		TArray<FObjectInfo>* ObjectInfos = ObjectsToReplicate.Find(ObjectToProcess.Object);
 		// ExtractReplicationDataForObject is supposed to be called in response to ForEachPendingObject... so either the call was invalid or ForEachPendingObject lied
 		if (!ensure(ObjectInfos))
 		{
@@ -217,7 +222,7 @@ namespace UE::ConcertSyncClient::Replication
 			return false;
 		}
 
-		const FObjectInfo& ObjectInfo = (*ObjectInfos)[StreamIndex];
+		FObjectInfo& ObjectInfo = (*ObjectInfos)[StreamIndex];
 		ConcertSyncCore::FReplicationPropertyFilter Filter(ObjectInfo.SelectedProperties);
 		TOptional<FConcertSessionSerializedPayload> Payload = ReplicationFormat.CreateReplicationEvent(
 			*Object,
@@ -229,6 +234,7 @@ namespace UE::ConcertSyncClient::Replication
 		if (Payload)
 		{
 			ProcessMoveable(MoveTemp(*Payload));
+			++ObjectInfo.ReplicationSequenceId;
 		}
 		return true;
 	}
diff --git a/Engine/Plugins/Developer/Concert/ConcertSync/ConcertSyncClient/Source/ConcertSyncClient/Private/Replication/Processing/ClientReplicationDataCollector.h b/Engine/Plugins/Developer/Concert/ConcertSync/ConcertSyncClient/Source/ConcertSyncClient/Private/Replication/Processing/ClientReplicationDataCollector.h
--- a/Engine/Plugins/Developer/Concert/ConcertSync/ConcertSyncClient/Source/ConcertSyncClient/Private/Replication/Processing/ClientReplicationDataCollector.h
+++ b/Engine/Plugins/Developer/Concert/ConcertSync/ConcertSyncClient/Source/ConcertSyncClient/Private/Replication/Processing/ClientReplicationDataCollector.h
@@ -72,7 +72,7 @@ namespace UE::ConcertSyncClient::Replication
 		void AppendOwningStreamsForObject(const FSoftObjectPath& ObjectPath, TSet<FGuid>& Paths) const;
 
 		//~ Begin IReplicationDataSource Interface
-		virtual void ForEachPendingObject(TFunctionRef<void(const FConcertReplicatedObjectId&)> ProcessItemFunc) const override;
+		virtual void ForEachPendingObject(TFunctionRef<void(const ConcertSyncCore::FPendingObjectReplicationInfo&)> ProcessItemFunc) const override;
 		virtual int32 NumObjects() const override { return NumTrackedObjects; }
 		virtual bool ExtractReplicationDataForObject(const FConcertReplicatedObjectId& Object, TFunctionRef<void(const FConcertSessionSerializedPayload& Payload)> ProcessCopyable, TFunctionRef<void(FConcertSessionSerializedPayload&& Payload)> ProcessMoveable) override;
 		//~ End IReplicationDataSource Interface
@@ -99,6 +99,9 @@ namespace UE::ConcertSyncClient::Replication
 			FGuid StreamId;
 			/** The properties to replicate */
 			FConcertPropertySelection SelectedProperties;
+			
+			/** Incremented every time replication data is sent out. Used for performance tracing. */
+			ConcertSyncCore::FSequenceId ReplicationSequenceId = 0;
 		};
 		
 		/** The objects and their properties to replicate */
diff --git a/Engine/Plugins/Developer/Concert/ConcertSync/ConcertSyncClient/Source/ConcertSyncClient/Private/Replication/Processing/ObjectReplicationApplierProcessor.cpp b/Engine/Plugins/Developer/Concert/ConcertSync/ConcertSyncClient/Source/ConcertSyncClient/Private/Replication/Processing/ObjectReplicationApplierProcessor.cpp
--- a/Engine/Plugins/Developer/Concert/ConcertSync/ConcertSyncClient/Source/ConcertSyncClient/Private/Replication/Processing/ObjectReplicationApplierProcessor.cpp
+++ b/Engine/Plugins/Developer/Concert/ConcertSync/ConcertSyncClient/Source/ConcertSyncClient/Private/Replication/Processing/ObjectReplicationApplierProcessor.cpp
@@ -7,6 +7,7 @@
 #include "Replication/Formats/IObjectReplicationFormat.h"
 
 #include "Components/SceneComponent.h"
+#include "Trace/ConcertProtocolTrace.h"
 
 namespace UE::ConcertSyncClient::Replication
 {
@@ -22,16 +23,18 @@ namespace UE::ConcertSyncClient::Replication
 
 	void FObjectReplicationApplierProcessor::ProcessObject(const FObjectProcessArgs& Args)
 	{
-		UObject* Object = ReplicationBridge.FindObjectIfAvailable(Args.ObjectInfo.Object);
+		CONCERT_TRACE_REPLICATION_OBJECT_SCOPE(ApplyReceivedObject, Args.ObjectInfo.ObjectId.Object, Args.ObjectInfo.SequenceId);
+		UObject* Object = ReplicationBridge.FindObjectIfAvailable(Args.ObjectInfo.ObjectId.Object);
 		if (!Object)
 		{
-			UE_LOG(LogConcert, Error, TEXT("Replication: Object %s is unavailable. The data source should not have reported it."), *Args.ObjectInfo.Object.ToString());
+			UE_LOG(LogConcert, Error, TEXT("Replication: Object %s is unavailable. The data source should not have reported it."), *Args.ObjectInfo.ObjectId.Object.ToString());
 			return;
 		}
 
 		bool bAppliedData = false;
-		GetDataSource().ExtractReplicationDataForObject(Args.ObjectInfo, [this, Object, &bAppliedData](const FConcertSessionSerializedPayload& Payload)
+		GetDataSource().ExtractReplicationDataForObject(Args.ObjectInfo.ObjectId, [this, &Args, Object, &bAppliedData](const FConcertSessionSerializedPayload& Payload)
 		{
+			CONCERT_TRACE_REPLICATION_OBJECT_SCOPE(SerializeReceivedObject, Args.ObjectInfo.ObjectId.Object, Args.ObjectInfo.SequenceId);
 			bAppliedData = true;
 
 			// TODO DP UE-193659: This is very hacky and leaves performance on the table... this is in case ApplyReplicationEvent updates the transform
@@ -45,7 +48,9 @@ namespace UE::ConcertSyncClient::Replication
 				ReplicationFormat.ApplyReplicationEvent(*Object, Payload);
 			}
 		});
-		// This should not happen. If it does, we're wasting  network bandwidth.
-		UE_CLOG(!bAppliedData, LogConcert, Warning, TEXT("Replication: Server sent data that could not be applied (likely it was empty) for object %s from stream %s"), *Args.ObjectInfo.Object.ToString(), *Args.ObjectInfo.StreamId.ToString());
+		
+		// This should not happen. If it does, we're wasting network bandwidth.
+		UE_CLOG(!bAppliedData, LogConcert, Warning, TEXT("Replication: Server sent data that could not be applied (likely it was empty) for object %s from stream %s"), *Args.ObjectInfo.ObjectId.Object.ToString(), *Args.ObjectInfo.ObjectId.StreamId.ToString());
+		CONCERT_TRACE_REPLICATION_OBJECT_SINK(Processed, Args.ObjectInfo.ObjectId.Object, Args.ObjectInfo.SequenceId);
 	}
 }
diff --git a/Engine/Plugins/Developer/Concert/ConcertSync/ConcertSyncCore/Source/ConcertSyncCore/ConcertSyncCore.Build.cs b/Engine/Plugins/Developer/Concert/ConcertSync/ConcertSyncCore/Source/ConcertSyncCore/ConcertSyncCore.Build.cs
--- a/Engine/Plugins/Developer/Concert/ConcertSync/ConcertSyncCore/Source/ConcertSyncCore/ConcertSyncCore.Build.cs
+++ b/Engine/Plugins/Developer/Concert/ConcertSync/ConcertSyncCore/Source/ConcertSyncCore/ConcertSyncCore.Build.cs
@@ -21,6 +21,7 @@ namespace UnrealBuildTool.Rules
 					"ConcertTransport",
 					"Json",
 					"SQLiteCore",
+					"TraceLog"
 				}
 			);
 
diff --git a/Engine/Plugins/Developer/Concert/ConcertSync/ConcertSyncCore/Source/ConcertSyncCore/Private/Replication/Processing/ObjectReplicationCache.cpp b/Engine/Plugins/Developer/Concert/ConcertSync/ConcertSyncCore/Source/ConcertSyncCore/Private/Replication/Processing/ObjectReplicationCache.cpp
--- a/Engine/Plugins/Developer/Concert/ConcertSync/ConcertSyncCore/Source/ConcertSyncCore/Private/Replication/Processing/ObjectReplicationCache.cpp
+++ b/Engine/Plugins/Developer/Concert/ConcertSync/ConcertSyncCore/Source/ConcertSyncCore/Private/Replication/Processing/ObjectReplicationCache.cpp
@@ -12,40 +12,31 @@ namespace UE::ConcertSyncCore
 		: ReplicationFormat(ReplicationFormat)
 	{}
 
-	int32 FObjectReplicationCache::StoreUntilConsumed(const FGuid& SendingEndpointId, const FGuid& OriginStreamId, const FConcertReplication_ObjectReplicationEvent& ObjectReplicationEvent)
+	FCacheStoreStats FObjectReplicationCache::StoreUntilConsumed(const FGuid& SendingEndpointId, const FGuid& OriginStreamId, const FSequenceId SequenceId, const FConcertReplication_ObjectReplicationEvent& ObjectReplicationEvent)
 	{
-		int32 NumAccepted = 0;
+		FCacheStoreStats Stats;
 		
 		const FConcertReplicatedObjectId ObjectId{ { OriginStreamId, ObjectReplicationEvent.ReplicatedObject }, SendingEndpointId };
 		FObjectCache* ObjectCacheBeforeAddition = Cache.Find(ObjectId);
 		if (ObjectCacheBeforeAddition)
 		{
-			// It is important to compare by address and not by TWeakPtr instance because each user has a different TWeakPtr instance (but they all point to the same memory).
-			TSet<FConcertReplication_ObjectReplicationEvent*> CombineOnceDetection;
-			for (const TPair<TWeakPtr<IReplicationCacheUser>, TWeakPtr<FConcertReplication_ObjectReplicationEvent>>& InUseDataPair : ObjectCacheBeforeAddition->DataInUse)
-			{
-				const TWeakPtr<FConcertReplication_ObjectReplicationEvent>& EventData = InUseDataPair.Value;
-				const TSharedPtr<FConcertReplication_ObjectReplicationEvent> EventDataPin = EventData.Pin();
-				if (EventDataPin && !CombineOnceDetection.Contains(EventDataPin.Get()))
-				{
-					CombineOnceDetection.Add(EventDataPin.Get());
-					ReplicationFormat.CombineReplicationEvents(EventDataPin->SerializedPayload, ObjectReplicationEvent.SerializedPayload);
-				}
-			}
+			Stats.NumCacheUpdates = CombineCachedDataWithNewData(ObjectId, SequenceId, ObjectReplicationEvent, *ObjectCacheBeforeAddition);
 		}
 
 		FObjectCache* ObjectCacheAfterAddition = ObjectCacheBeforeAddition ? ObjectCacheBeforeAddition : nullptr;
 		TSharedPtr<FConcertReplication_ObjectReplicationEvent> LazilyCopiedEventPtr;
 		for (const TSharedRef<IReplicationCacheUser>& CacheUser : CacheUsers)
 		{
-			if (ObjectCacheBeforeAddition && ObjectCacheBeforeAddition->DataInUse.Contains(CacheUser))
+			if (ObjectCacheBeforeAddition
+				// They've already had OnDataCached called if they're contained
+				&& ObjectCacheBeforeAddition->DataInUse.Contains(CacheUser))
 			{
 				continue;
 			}
 
 			if (CacheUser->WantsToAcceptObject(ObjectId))
 			{
-				++NumAccepted;
+				++Stats.NumInsertions;
 				
 				// Do the event copy only when somebody wants the data...
 				if (!LazilyCopiedEventPtr.IsValid())
@@ -83,7 +74,7 @@ namespace UE::ConcertSyncCore
 				});
 				
 				const TWeakPtr<FConcertReplication_ObjectReplicationEvent> WeakGuardPtr = GuardPtr;
-				CacheUser->OnDataCached(ObjectId, MoveTemp(GuardPtr));
+				CacheUser->OnDataCached(ObjectId, SequenceId, MoveTemp(GuardPtr));
 				// Was it instantly consumed? Should not really happen but it could technically...
 				if (!LIKELY(WeakGuardPtr.IsValid()))
 				{
@@ -98,7 +89,7 @@ namespace UE::ConcertSyncCore
 			}
 		}
 
-		return NumAccepted;
+		return Stats;
 	}
 
 	void FObjectReplicationCache::RegisterDataCacheUser(TSharedRef<IReplicationCacheUser> User)
@@ -123,4 +114,43 @@ namespace UE::ConcertSyncCore
 			}
 		}
 	}
+
+	uint32 FObjectReplicationCache::CombineCachedDataWithNewData(
+		const FConcertReplicatedObjectId& ObjectId,
+		FSequenceId NewSequenceId,
+		const FConcertReplication_ObjectReplicationEvent& NewData,
+		const FObjectCache& ObjectCacheBeforeAddition
+		) const
+	{
+		uint32 NumUpdates = 0;
+		
+		// Caches users may point to the same or different memory depending how fast they process
+		TSet<FConcertReplication_ObjectReplicationEvent*> CombineOnceDetection;
+		for (const TPair<TWeakPtr<IReplicationCacheUser>, TWeakPtr<FConcertReplication_ObjectReplicationEvent>>& InUseDataPair : ObjectCacheBeforeAddition.DataInUse)
+		{
+			const TWeakPtr<IReplicationCacheUser> CacheUser = InUseDataPair.Key;
+			const TWeakPtr<FConcertReplication_ObjectReplicationEvent>& EventData = InUseDataPair.Value;
+
+			// There's no point in updating if either the user has destroyed itself (rudely without telling us) or the user has released that data already.
+			const TSharedPtr<IReplicationCacheUser> CacheUserPin = CacheUser.Pin();
+			const TSharedPtr<FConcertReplication_ObjectReplicationEvent> EventDataPin = EventData.Pin();
+			if (!EventDataPin || !CacheUserPin)
+			{
+				continue;
+			}
+
+			// Multiple cache users may be using the same data so only combine once ...
+			if (!CombineOnceDetection.Contains(EventDataPin.Get()))
+			{
+				CombineOnceDetection.Add(EventDataPin.Get());
+				ReplicationFormat.CombineReplicationEvents(EventDataPin->SerializedPayload, NewData.SerializedPayload);
+			}
+			// ... but let every user know that the data was combined with another SequenceId
+			CacheUserPin->OnCachedDataUpdated(ObjectId, NewSequenceId);
+
+			++NumUpdates;
+		}
+
+		return NumUpdates;
+	}
 }
diff --git a/Engine/Plugins/Developer/Concert/ConcertSync/ConcertSyncCore/Source/ConcertSyncCore/Private/Replication/Processing/ObjectReplicationProcessor.cpp b/Engine/Plugins/Developer/Concert/ConcertSync/ConcertSyncCore/Source/ConcertSyncCore/Private/Replication/Processing/ObjectReplicationProcessor.cpp
--- a/Engine/Plugins/Developer/Concert/ConcertSync/ConcertSyncCore/Source/ConcertSyncCore/Private/Replication/Processing/ObjectReplicationProcessor.cpp
+++ b/Engine/Plugins/Developer/Concert/ConcertSync/ConcertSyncCore/Source/ConcertSyncCore/Private/Replication/Processing/ObjectReplicationProcessor.cpp
@@ -16,7 +16,7 @@ namespace UE::ConcertSyncCore
 	void FObjectReplicationProcessor::ProcessObjects(const FProcessObjectsParams& Params)
 	{
 		// TODO UE-190714: Respect time budget and prioritize objects
-		DataSource.ForEachPendingObject([this](const FConcertReplicatedObjectId& ObjectInfo)
+		DataSource.ForEachPendingObject([this](const FPendingObjectReplicationInfo& ObjectInfo)
 		{
 			ProcessObject({ ObjectInfo });
 		});
diff --git a/Engine/Plugins/Developer/Concert/ConcertSync/ConcertSyncCore/Source/ConcertSyncCore/Private/Replication/Processing/ObjectReplicationReceiver.cpp b/Engine/Plugins/Developer/Concert/ConcertSync/ConcertSyncCore/Source/ConcertSyncCore/Private/Replication/Processing/ObjectReplicationReceiver.cpp
--- a/Engine/Plugins/Developer/Concert/ConcertSync/ConcertSyncCore/Source/ConcertSyncCore/Private/Replication/Processing/ObjectReplicationReceiver.cpp
+++ b/Engine/Plugins/Developer/Concert/ConcertSync/ConcertSyncCore/Source/ConcertSyncCore/Private/Replication/Processing/ObjectReplicationReceiver.cpp
@@ -8,9 +8,23 @@
 #include "Replication/Processing/ObjectReplicationCache.h"
 
 #include "HAL/IConsoleManager.h"
+#include "Trace/ConcertProtocolTrace.h"
 
 namespace UE::ConcertSyncCore
 {
+	namespace Private
+	{
+		static void TraceDroppedObjectEventIf(bool bSendTrace, const FConcertReplication_ObjectReplicationEvent& ObjectEvent)
+		{
+#if UE_CONCERT_TRACE_ENABLED
+			if (bSendTrace)
+			{
+				CONCERT_TRACE_REPLICATION_OBJECT_SINK(Dropped, ObjectEvent.ReplicatedObject, ObjectEvent.ReplicationSequenceId);
+			}
+#endif
+		}
+	}
+	
 	static TAutoConsoleVariable<bool> CVarLogReceivedObjects(TEXT("Concert.Replication.LogReceivedObjects"), false, TEXT("Enable Concert logging for received replicated objects."));
 	
 	FObjectReplicationReceiver::FObjectReplicationReceiver(IConcertSession& Session, FObjectReplicationCache& ReplicationCache)
@@ -28,10 +42,11 @@ namespace UE::ConcertSyncCore
 	void FObjectReplicationReceiver::HandleBatchReplicationEvent(const FConcertSessionContext& SessionContext, const FConcertReplication_BatchReplicationEvent& Event)
 	{
 		// Fyi: an object may have multiple changes in a batch replication event: each stream can modify different properties as long as they do not overlap.
-		int32 NumObjects = 0;
-		int32 NumRejectedObjectChanges = 0;
-		int32 NumCacheUsages = 0;
-		int32 NumOfAcceptedObjectChanges = 0;
+		uint32 NumObjects = 0;
+		uint32 NumRejectedObjectChanges = 0;
+		uint32 NumCacheInsertions = 0;
+		uint32 NumCacheUpdates = 0;
+		uint32 NumOfAcceptedObjectChanges = 0;
 		
 		for (const FConcertReplication_StreamReplicationEvent& StreamEvent : Event.Streams)
 		{
@@ -40,14 +55,20 @@ namespace UE::ConcertSyncCore
 			
 			for (const FConcertReplication_ObjectReplicationEvent& ObjectEvent : StreamEvent.ReplicatedObjects)
 			{
+				CONCERT_TRACE_REPLICATION_OBJECT_TRANSMISSION_RECEIVE(ObjectEvent.ReplicatedObject, ObjectEvent.ReplicationSequenceId);
 				if (ShouldAcceptObject(SessionContext, StreamEvent, ObjectEvent))
 				{
-					const int32 NumAccepted = ReplicationCache.StoreUntilConsumed(SessionContext.SourceEndpointId, StreamEvent.StreamId, ObjectEvent);
-					NumCacheUsages += NumAccepted;
-					NumOfAcceptedObjectChanges += NumAccepted == 0 ? 0 : 1;
+					CONCERT_TRACE_REPLICATION_OBJECT_SCOPE(EnqueueReceivedObject, ObjectEvent.ReplicatedObject, ObjectEvent.ReplicationSequenceId);
+					const FCacheStoreStats CacheStoreStats = ReplicationCache.StoreUntilConsumed(SessionContext.SourceEndpointId, StreamEvent.StreamId, ObjectEvent.ReplicationSequenceId, ObjectEvent);
+					NumCacheInsertions += CacheStoreStats.NumInsertions;
+					NumCacheUpdates += CacheStoreStats.NumCacheUpdates;
+					NumOfAcceptedObjectChanges += CacheStoreStats.NumInsertions == 0 ? 0 : 1;
+					
+					Private::TraceDroppedObjectEventIf(CacheStoreStats.NoChangesMade(), ObjectEvent);
 				}
 				else
 				{
+					CONCERT_TRACE_REPLICATION_OBJECT_SINK(Rejected, ObjectEvent.ReplicatedObject, ObjectEvent.ReplicationSequenceId);
 					++NumRejectedObjectChanges;
 				}
 			}
@@ -55,12 +76,13 @@ namespace UE::ConcertSyncCore
 
 		if (CVarLogReceivedObjects.GetValueOnGameThread())
 		{
-			UE_LOG(LogConcert, Log, TEXT("Received %d streams with %d object changes from endpoint %s. Cached %d object changes with a total of %d cache usages."),
+			UE_LOG(LogConcert, Log, TEXT("Received %u streams with %u object changes from endpoint %s. Cached %u object changes with a total of new %u cache insertions and %u cache updates."),
 				Event.Streams.Num(),
 				NumObjects,
 				*SessionContext.SourceEndpointId.ToString(),
 				NumOfAcceptedObjectChanges,
-				NumCacheUsages
+				NumCacheInsertions,
+				NumCacheUpdates
 			);
 			UE_CLOG(NumRejectedObjectChanges > 0, LogConcert, Warning, TEXT("Rejected %d object changes."), NumRejectedObjectChanges);
 		}
diff --git a/Engine/Plugins/Developer/Concert/ConcertSync/ConcertSyncCore/Source/ConcertSyncCore/Private/Replication/Processing/ObjectReplicationSender.cpp b/Engine/Plugins/Developer/Concert/ConcertSync/ConcertSyncCore/Source/ConcertSyncCore/Private/Replication/Processing/ObjectReplicationSender.cpp
--- a/Engine/Plugins/Developer/Concert/ConcertSync/ConcertSyncCore/Source/ConcertSyncCore/Private/Replication/Processing/ObjectReplicationSender.cpp
+++ b/Engine/Plugins/Developer/Concert/ConcertSync/ConcertSyncCore/Source/ConcertSyncCore/Private/Replication/Processing/ObjectReplicationSender.cpp
@@ -5,6 +5,7 @@
 #include "ConcertLogGlobal.h"
 #include "IConcertSession.h"
 #include "Replication/Processing/IReplicationDataSource.h"
+#include "Trace/ConcertProtocolTrace.h"
 
 #include "Algo/Accumulate.h"
 #include "HAL/IConsoleManager.h"
@@ -36,6 +37,7 @@ namespace UE::ConcertSyncCore
 				*TargetEndpointId.ToString()
 				);
 			
+			TraceStartSendingMarkedObjects(); 
 			Session.SendCustomEvent(EventToSend, TargetEndpointId,
 				// Replication is always unreliable - if it fails to deliver we'll send updated data soon again
 				// TODO: In regular intervals send CRC values to detect that a change is missing
@@ -50,19 +52,48 @@ namespace UE::ConcertSyncCore
 
 	void FObjectReplicationSender::ProcessObject(const FObjectProcessArgs& Args)
 	{
-		const FSoftObjectPath& ReplicatedObject = Args.ObjectInfo.Object;
+		// It would be easier to just CONCERT_TRACE_REPLICATION_OBJECT_TRANSMISSION_START here but for better precision we must postbone it until all objects have been processed
+		MarkObjectForTrace(Args.ObjectInfo.ObjectId, Args.ObjectInfo.SequenceId);
+		CONCERT_TRACE_REPLICATION_OBJECT_SCOPE(CollectObjectDataForSend, Args.ObjectInfo.ObjectId.Object, Args.ObjectInfo.SequenceId);
+		
+		const FSoftObjectPath& ReplicatedObject = Args.ObjectInfo.ObjectId.Object;
 		auto CaptureData = [this, &Args, &ReplicatedObject]<typename TPayloadRefType>(TPayloadRefType&& Payload)
 		{
-			const int32 PreexistingIndex = EventToSend.Streams.IndexOfByPredicate([&Args](const FConcertReplication_StreamReplicationEvent& StreamData){ return StreamData.StreamId == Args.ObjectInfo.StreamId; });
+			const int32 PreexistingIndex = EventToSend.Streams.IndexOfByPredicate([&Args](const FConcertReplication_StreamReplicationEvent& StreamData){ return StreamData.StreamId == Args.ObjectInfo.ObjectId.StreamId; });
 			FConcertReplication_StreamReplicationEvent& StreamData = EventToSend.Streams.IsValidIndex(PreexistingIndex)
 				? EventToSend.Streams[PreexistingIndex]
-				: EventToSend.Streams[EventToSend.Streams.Emplace(Args.ObjectInfo.StreamId)];
+				: EventToSend.Streams[EventToSend.Streams.Emplace(Args.ObjectInfo.ObjectId.StreamId)];
 			StreamData.ReplicatedObjects.Add({
 				ReplicatedObject,
+				static_cast<int32>(Args.ObjectInfo.SequenceId), // This is silly... Blueprints do not support uint32 so we need to pretend it's an int32.
 				// Take advantage of move semantics if it is possible - this depends on how our data source internally obtains its payloads
 				Forward<TPayloadRefType>(Payload)
 			});
 		};
-		GetDataSource().ExtractReplicationDataForObject(Args.ObjectInfo, CaptureData, CaptureData);
+		GetDataSource().ExtractReplicationDataForObject(Args.ObjectInfo.ObjectId, CaptureData, CaptureData);
+	}
+
+	inline void FObjectReplicationSender::MarkObjectForTrace(const FConcertReplicatedObjectId& Object, FSequenceId Id)
+	{
+#if UE_CONCERT_TRACE_ENABLED
+		if (ConcertTrace::IsTracingReplication())
+		{
+			ObjectsToTraceThisFrame.Add(Object, Id);
+		}
+#endif
+	}
+
+	inline void FObjectReplicationSender::TraceStartSendingMarkedObjects()
+	{
+#if UE_CONCERT_TRACE_ENABLED
+		if (ConcertTrace::IsTracingReplication())
+		{
+			for (const TPair<FConcertReplicatedObjectId, FSequenceId>& SendingNow : ObjectsToTraceThisFrame)
+			{
+				CONCERT_TRACE_REPLICATION_OBJECT_TRANSMISSION_START(TargetEndpointId, SendingNow.Key.Object, SendingNow.Value);
+			}
+			ObjectsToTraceThisFrame.Reset();
+		}
+#endif
 	}
 }
diff --git a/Engine/Plugins/Developer/Concert/ConcertSync/ConcertSyncCore/Source/ConcertSyncCore/Private/Replication/Processing/ReplicationDataQueuer.cpp b/Engine/Plugins/Developer/Concert/ConcertSync/ConcertSyncCore/Source/ConcertSyncCore/Private/Replication/Processing/ReplicationDataQueuer.cpp
--- a/Engine/Plugins/Developer/Concert/ConcertSync/ConcertSyncCore/Source/ConcertSyncCore/Private/Replication/Processing/ReplicationDataQueuer.cpp
+++ b/Engine/Plugins/Developer/Concert/ConcertSync/ConcertSyncCore/Source/ConcertSyncCore/Private/Replication/Processing/ReplicationDataQueuer.cpp
@@ -2,15 +2,17 @@
 
 #include "Replication/Processing/ReplicationDataQueuer.h"
 
+#include "ConcertLogGlobal.h"
 #include "Replication/Messages/ObjectReplication.h"
+#include "Trace/ConcertProtocolTrace.h"
 
 namespace UE::ConcertSyncCore
 {
-	void FReplicationDataQueuer::ForEachPendingObject(TFunctionRef<void(const FConcertReplicatedObjectId&)> ProcessItemFunc) const
+	void FReplicationDataQueuer::ForEachPendingObject(TFunctionRef<void(const FPendingObjectReplicationInfo&)> ProcessItemFunc) const
 	{
 		for (auto It = PendingEvents.CreateConstIterator(); It; ++It)
 		{
-			ProcessItemFunc(It->Key);
+			ProcessItemFunc({ It->Key, It->Value.SequenceId });
 		}
 	}
 
@@ -25,23 +27,36 @@ namespace UE::ConcertSyncCore
 		TFunctionRef<void(FConcertSessionSerializedPayload&& Payload)> ProcessMoveable
 		)
 	{
-		TSharedPtr<const FConcertReplication_ObjectReplicationEvent> Event;
-		const bool bSuccess = PendingEvents.RemoveAndCopyValue(Object, Event);
+		FPendingObjectData EventData;
+		const bool bSuccess = PendingEvents.RemoveAndCopyValue(Object, EventData);
 		if (!ensureMsgf(bSuccess, TEXT("ExtractReplicationDataForObject for an item that was not returned by ForEachPendingObject")))
 		{
 			return false;
 		}
 
 		// Event may be shared by other FReplicationDataQueuers since it originates from the replication cache, so sadly no move.
-		ProcessCopyable(Event->SerializedPayload);
+		ProcessCopyable(EventData.DataToApply->SerializedPayload);
 		return true;
 	}
 
-	void FReplicationDataQueuer::OnDataCached(const FConcertReplicatedObjectId& Object, TSharedRef<const FConcertReplication_ObjectReplicationEvent> Data)
+	void FReplicationDataQueuer::OnDataCached(const FConcertReplicatedObjectId& Object, FSequenceId SequenceId, TSharedRef<const FConcertReplication_ObjectReplicationEvent> Data)
 	{
-		PendingEvents.Add(Object, MoveTemp(Data));
+		PendingEvents.Add(Object, FPendingObjectData{ Data, SequenceId });
 	}
-	
+
+	void FReplicationDataQueuer::OnCachedDataUpdated(const FConcertReplicatedObjectId& Object, FSequenceId SequenceId)
+	{
+		FPendingObjectData* ObjectData = PendingEvents.Find(Object);
+		if (ensureMsgf(ObjectData, TEXT("OnCachedDataUpdated called for %s even though we have no cached data. Investigate in correct API call!"), *Object.ToString()))
+		{
+			// TODO DP: This may want to be a specialized macro so Insights can visually highlight that a packet was merged.
+			// Trace that the old data has "finished" sending ... 
+			CONCERT_TRACE_REPLICATION_OBJECT_SINK(Merged, Object.Object, ObjectData->SequenceId);
+			// ... since it is merged
+			ObjectData->SequenceId = SequenceId;
+		}
+	}
+
 	void FReplicationDataQueuer::BindToCache(FObjectReplicationCache& InReplicationCache)
 	{
 		ReplicationCache = &InReplicationCache;
diff --git a/Engine/Plugins/Developer/Concert/ConcertSync/ConcertSyncCore/Source/ConcertSyncCore/Public/Replication/Data/SequenceId.h b/Engine/Plugins/Developer/Concert/ConcertSync/ConcertSyncCore/Source/ConcertSyncCore/Public/Replication/Data/SequenceId.h
new file mode 100644
--- /dev/null
+++ b/Engine/Plugins/Developer/Concert/ConcertSync/ConcertSyncCore/Source/ConcertSyncCore/Public/Replication/Data/SequenceId.h
@@ -0,0 +1,15 @@
+// Copyright Epic Games, Inc. All Rights Reserved.
+
+#pragma once
+
+#include "HAL/Platform.h"
+
+namespace UE::ConcertSyncCore
+{
+	/**
+	 * A sequence ID is associated with every update made to a replicated object.
+	 * Every time a replicated change is sent, the ID is incremented by the original data source.
+	 * This is used primarily for tracing object replication performance.
+	 */
+	using FSequenceId = uint32;
+}
\ No newline at end of file
diff --git a/Engine/Plugins/Developer/Concert/ConcertSync/ConcertSyncCore/Source/ConcertSyncCore/Public/Replication/Messages/ObjectReplication.h b/Engine/Plugins/Developer/Concert/ConcertSync/ConcertSyncCore/Source/ConcertSyncCore/Public/Replication/Messages/ObjectReplication.h
--- a/Engine/Plugins/Developer/Concert/ConcertSync/ConcertSyncCore/Source/ConcertSyncCore/Public/Replication/Messages/ObjectReplication.h
+++ b/Engine/Plugins/Developer/Concert/ConcertSync/ConcertSyncCore/Source/ConcertSyncCore/Public/Replication/Messages/ObjectReplication.h
@@ -21,6 +21,10 @@ struct FConcertReplication_ObjectReplicationEvent
 	UPROPERTY()
 	FSoftObjectPath ReplicatedObject;
 	
+	/** Incremented every time replication data is sent out. Used for performance tracing. */
+	UPROPERTY()
+	int32 ReplicationSequenceId = -1;
+	
 	/** Contains another struct as payload, such as FConcertFullObjectReplicationData. */
 	UPROPERTY()
 	FConcertSessionSerializedPayload SerializedPayload;
diff --git a/Engine/Plugins/Developer/Concert/ConcertSync/ConcertSyncCore/Source/ConcertSyncCore/Public/Replication/Processing/IReplicationDataSource.h b/Engine/Plugins/Developer/Concert/ConcertSync/ConcertSyncCore/Source/ConcertSyncCore/Public/Replication/Processing/IReplicationDataSource.h
--- a/Engine/Plugins/Developer/Concert/ConcertSync/ConcertSyncCore/Source/ConcertSyncCore/Public/Replication/Processing/IReplicationDataSource.h
+++ b/Engine/Plugins/Developer/Concert/ConcertSync/ConcertSyncCore/Source/ConcertSyncCore/Public/Replication/Processing/IReplicationDataSource.h
@@ -3,14 +3,26 @@
 #pragma once
 
 #include "ConcertMessageData.h"
+#include "Replication/Data/ObjectIds.h"
+#include "Replication/Data/SequenceId.h"
+
 #include "HAL/Platform.h"
 #include "Misc/Optional.h"
 #include "Templates/Function.h"
 
-struct FConcertReplicatedObjectId;
-
 namespace UE::ConcertSyncCore
 {
+	struct FPendingObjectReplicationInfo
+	{
+		/** The object to be replicated */
+		FConcertReplicatedObjectId ObjectId;
+		/**
+		 * The ID of the change that is pending.
+		 * If previous changes were combined, this is the ID of the highest combined change.
+		 */
+		FSequenceId SequenceId;
+	};
+	
 	/**
 	 * Responsible for obtaining replication data.
 	 * 
@@ -23,13 +35,15 @@ namespace UE::ConcertSyncCore
 	public:
 		
 		/** Iterates the objects that must be processed for replication (the result of ExtractReplicationDataForObject MAY return something new). */
-		virtual void ForEachPendingObject(TFunctionRef<void(const FConcertReplicatedObjectId&)> ProcessItemFunc) const = 0;
+		virtual void ForEachPendingObject(TFunctionRef<void(const FPendingObjectReplicationInfo&)> ProcessItemFunc) const = 0;
 		/** @return The number of objects ForEachPendingObject would iterate. Can be used, e.g. for Reserve()-ing a container. */
 		virtual int32 NumObjects() const = 0;
 
 		/**
 		 * Extracts data for Object. Object must have been previously returned by ForEachPendingObject.
+		 * 
 		 * This "dequeues" Object so it will not be mentioned by ForEachPendingObject again until it is again marked as "dirty".
+		 * This might also update the SequenceId associated with the object (if this data source generates data as opposed to queuing it).
 		 *
 		 * The call to ProcessCopyable / ProcessMoveable may be skipped if there is no new data to send.
 		 * Either ProcessCopyable or ProcessMoveable will be called, never both, and it will be called at most once.
@@ -49,21 +63,7 @@ namespace UE::ConcertSyncCore
 		{
 			return ExtractReplicationDataForObject(Object, ProcessCopyable, [&ProcessCopyable](FConcertSessionSerializedPayload&& Payload){ ProcessCopyable(Payload); });
 		}
-
-		/** Util function for add the contained objects to a container (TArray / TSet). Container must have Reserve and Add functions.  */
-		template<typename TContainerType>
-		void AddEachObjectToContainer(TContainerType& OutContainer) const;
 		
 		virtual ~IReplicationDataSource() = default;
 	};
-
-	template <typename TContainerType>
-	void IReplicationDataSource::AddEachObjectToContainer(TContainerType& OutContainer) const
-	{
-		OutContainer.Reserve(OutContainer.Num() + NumObjects());
-		ForEachPendingObject([&OutContainer](const FSoftObjectPath& Item) mutable
-		{
-			OutContainer.Add(Item);
-		});
-	}
 }
diff --git a/Engine/Plugins/Developer/Concert/ConcertSync/ConcertSyncCore/Source/ConcertSyncCore/Public/Replication/Processing/ObjectReplicationCache.h b/Engine/Plugins/Developer/Concert/ConcertSync/ConcertSyncCore/Source/ConcertSyncCore/Public/Replication/Processing/ObjectReplicationCache.h
--- a/Engine/Plugins/Developer/Concert/ConcertSync/ConcertSyncCore/Source/ConcertSyncCore/Public/Replication/Processing/ObjectReplicationCache.h
+++ b/Engine/Plugins/Developer/Concert/ConcertSync/ConcertSyncCore/Source/ConcertSyncCore/Public/Replication/Processing/ObjectReplicationCache.h
@@ -2,9 +2,11 @@
 
 #pragma once
 
+#include "Replication/Data/ObjectIds.h"
+#include "Replication/Data/SequenceId.h"
+
 #include "Containers/Array.h"
 #include "Delegates/DelegateCombinations.h"
-#include "Replication/Data/ObjectIds.h"
 #include "Templates/Function.h"
 #include "Templates/SharedPointer.h"
 
@@ -31,13 +33,28 @@ namespace UE::ConcertSyncCore
 		 * Called when data that is interesting to this user becomes available.
 		 *
 		 * The user can keep hold of Data until it is used, at which point it just let's Data get out of scope.
-		 * If new data is received while this user is referencing Data, Data will be combined to contain any new data.
+		 * If new data is received while this user is referencing Data, the Data will be combined and OnCachedDataUpdated called.
 		 */
-		virtual void OnDataCached(const FConcertReplicatedObjectId& Object, TSharedRef<const FConcertReplication_ObjectReplicationEvent> Data) = 0;
+		virtual void OnDataCached(const FConcertReplicatedObjectId& Object, FSequenceId SequenceId, TSharedRef<const FConcertReplication_ObjectReplicationEvent> Data) = 0;
+
+		/**
+		 * This function is called when new data is received while this cache user is still holding on to an object previously received with OnDataCached.
+		 * @param Object The object for which data was updated.
+		 * @param SequenceId The ID of the event data that was combined
+		 */
+		virtual void OnCachedDataUpdated(const FConcertReplicatedObjectId& Object, FSequenceId SequenceId) {}
 
 		virtual ~IReplicationCacheUser() = default;
 	};
 	
+	struct FCacheStoreStats
+	{
+		uint32 NumInsertions = 0;
+		uint32 NumCacheUpdates = 0;
+
+		bool NoChangesMade() const { return NumInsertions == 0 && NumCacheUpdates == 0; }
+	};
+	
 	/**
 	 * This is an intermediate place for received data to live before it is further processed.
 	 *
@@ -55,15 +72,16 @@ namespace UE::ConcertSyncCore
 	public:
 
 		FObjectReplicationCache(IObjectReplicationFormat& ReplicationFormat UE_LIFETIMEBOUND);
-
+		
 		/**
 		 * Called when new data is received for an object and shares it with any IObjectCacheUser that is possibly interested in it.
 		 * @param SendingEndpointId The ID of the client endpoint that sent this data
 		 * @param OriginStreamId The stream from which the object was replicated
+		 * @param SequenceId The ID of the change. Used primarily for performance tracing.
 		 * @param ObjectReplicationEvent The data that was replicated
 		 * @return The number of cache users that accepted this event. 
 		 */
-		int32 StoreUntilConsumed(const FGuid& SendingEndpointId, const FGuid& OriginStreamId, const FConcertReplication_ObjectReplicationEvent& ObjectReplicationEvent);
+		FCacheStoreStats StoreUntilConsumed(const FGuid& SendingEndpointId, const FGuid& OriginStreamId, const FSequenceId SequenceId, const FConcertReplication_ObjectReplicationEvent& ObjectReplicationEvent);
 
 		/** Registers a new user, which will start receiving data for any new data received from now on. */
 		void RegisterDataCacheUser(TSharedRef<IReplicationCacheUser> User);
@@ -92,6 +110,17 @@ namespace UE::ConcertSyncCore
 			TMap<TWeakPtr<IReplicationCacheUser>, TWeakPtr<FConcertReplication_ObjectReplicationEvent>> DataInUse;
 		};
 		/** Maps every object to the events cached for it. */
-		TMap<FConcertObjectInStreamID, FObjectCache> Cache; 
+		TMap<FConcertObjectInStreamID, FObjectCache> Cache;
+		
+		/**
+		 * Combines old data that cache users may have associated for ObjectId with the NewData and notifies the users with OnCachedDataUpdated.
+		 * @return Number of cache users that were updated
+		 */
+		uint32 CombineCachedDataWithNewData(
+			const FConcertReplicatedObjectId& ObjectId,
+			FSequenceId NewSequenceId,
+			const FConcertReplication_ObjectReplicationEvent& NewData,
+			const FObjectCache& ObjectCacheBeforeAddition
+			) const;
 	};
 }
\ No newline at end of file
diff --git a/Engine/Plugins/Developer/Concert/ConcertSync/ConcertSyncCore/Source/ConcertSyncCore/Public/Replication/Processing/ObjectReplicationProcessor.h b/Engine/Plugins/Developer/Concert/ConcertSync/ConcertSyncCore/Source/ConcertSyncCore/Public/Replication/Processing/ObjectReplicationProcessor.h
--- a/Engine/Plugins/Developer/Concert/ConcertSync/ConcertSyncCore/Source/ConcertSyncCore/Public/Replication/Processing/ObjectReplicationProcessor.h
+++ b/Engine/Plugins/Developer/Concert/ConcertSync/ConcertSyncCore/Source/ConcertSyncCore/Public/Replication/Processing/ObjectReplicationProcessor.h
@@ -26,21 +26,6 @@ namespace UE::ConcertSyncCore
 		// TODO UE-190714: We should add a time budget.
 	};
 	
-	
-	class CONCERTSYNCCORE_API IObjectReplicationProcessor
-	{
-	public:
-		
-		/**
-		 * Processes all changed objects under the given time budget.
-		 * TODO UE-190714: A time budget should be added.
-		 */
-		virtual void ProcessObjects(const FProcessObjectsParams& Params) = 0;
-
-		virtual ~IObjectReplicationProcessor() = default;
-	};
-
-
 	/**
 	 * Responsible for prioritizing a list of objects and processing them.
 	 * 
@@ -68,11 +53,11 @@ namespace UE::ConcertSyncCore
 		struct FObjectProcessArgs
 		{
 			/** Info about the object to process */
-			FConcertReplicatedObjectId ObjectInfo;
+			FPendingObjectReplicationInfo ObjectInfo;
 		};
 
 		FORCEINLINE IReplicationDataSource& GetDataSource() const { return DataSource; }
-		
+
 		/** Processes the object. */
 		virtual void ProcessObject(const FObjectProcessArgs& Args) = 0;
 
diff --git a/Engine/Plugins/Developer/Concert/ConcertSync/ConcertSyncCore/Source/ConcertSyncCore/Public/Replication/Processing/ObjectReplicationSender.h b/Engine/Plugins/Developer/Concert/ConcertSync/ConcertSyncCore/Source/ConcertSyncCore/Public/Replication/Processing/ObjectReplicationSender.h
--- a/Engine/Plugins/Developer/Concert/ConcertSync/ConcertSyncCore/Source/ConcertSyncCore/Public/Replication/Processing/ObjectReplicationSender.h
+++ b/Engine/Plugins/Developer/Concert/ConcertSync/ConcertSyncCore/Source/ConcertSyncCore/Public/Replication/Processing/ObjectReplicationSender.h
@@ -4,6 +4,7 @@
 
 #include "ObjectReplicationProcessor.h"
 #include "Replication/Messages/ObjectReplication.h"
+#include "Trace/ConcertTraceConfig.h"
 
 class IConcertSession;
 
@@ -45,7 +46,15 @@ namespace UE::ConcertSyncCore
 		/** The session through which replication messages are sent. */
 		IConcertSession& Session;
 
+#if UE_CONCERT_TRACE_ENABLED
+		/** Set by ProcessObject. Used when we actually start sending the data. */
+		TMap<FConcertReplicatedObjectId, FSequenceId> ObjectsToTraceThisFrame;
+#endif
+
 		/** This event is filled in ProcessObjects and finally sent to TargetEndpointId. */
 		FConcertReplication_BatchReplicationEvent EventToSend;
+
+		void MarkObjectForTrace(const FConcertReplicatedObjectId& Object, FSequenceId Id);
+		void TraceStartSendingMarkedObjects();
 	};
 }
diff --git a/Engine/Plugins/Developer/Concert/ConcertSync/ConcertSyncCore/Source/ConcertSyncCore/Public/Replication/Processing/Proxy/ObjectProcessorProxy_Frequency.h b/Engine/Plugins/Developer/Concert/ConcertSync/ConcertSyncCore/Source/ConcertSyncCore/Public/Replication/Processing/Proxy/ObjectProcessorProxy_Frequency.h
--- a/Engine/Plugins/Developer/Concert/ConcertSync/ConcertSyncCore/Source/ConcertSyncCore/Public/Replication/Processing/Proxy/ObjectProcessorProxy_Frequency.h
+++ b/Engine/Plugins/Developer/Concert/ConcertSync/ConcertSyncCore/Source/ConcertSyncCore/Public/Replication/Processing/Proxy/ObjectProcessorProxy_Frequency.h
@@ -83,7 +83,7 @@ namespace UE::ConcertSyncCore
 	template <TIsObjectProcessorConcept TRealProcessorImpl>
 	void TObjectProcessorProxy_Frequency<TRealProcessorImpl>::ProcessObject(const FObjectReplicationProcessor::FObjectProcessArgs& Args)
 	{
-		const FConcertReplicatedObjectId& Object = Args.ObjectInfo;
+		const FConcertReplicatedObjectId& Object = Args.ObjectInfo.ObjectId;
 		const FConcertObjectReplicationSettings FrequencySettings = GetObjectFrequencySettingsDelegate.Execute(Object);
 
 		// Throttle calls to the real ProcessObject implementation based on the object's frequency settings
diff --git a/Engine/Plugins/Developer/Concert/ConcertSync/ConcertSyncCore/Source/ConcertSyncCore/Public/Replication/Processing/ReplicationDataQueuer.h b/Engine/Plugins/Developer/Concert/ConcertSync/ConcertSyncCore/Source/ConcertSyncCore/Public/Replication/Processing/ReplicationDataQueuer.h
--- a/Engine/Plugins/Developer/Concert/ConcertSync/ConcertSyncCore/Source/ConcertSyncCore/Public/Replication/Processing/ReplicationDataQueuer.h
+++ b/Engine/Plugins/Developer/Concert/ConcertSync/ConcertSyncCore/Source/ConcertSyncCore/Public/Replication/Processing/ReplicationDataQueuer.h
@@ -25,13 +25,14 @@ namespace UE::ConcertSyncCore
 	public:
 
 		//~ Begin IReplicationDataSource Interface
-		virtual void ForEachPendingObject(TFunctionRef<void(const FConcertReplicatedObjectId&)> ProcessItemFunc) const override;
+		virtual void ForEachPendingObject(TFunctionRef<void(const FPendingObjectReplicationInfo&)> ProcessItemFunc) const override;
 		virtual int32 NumObjects() const override;
 		virtual bool ExtractReplicationDataForObject(const FConcertReplicatedObjectId& Object, TFunctionRef<void(const FConcertSessionSerializedPayload& Payload)> ProcessCopyable, TFunctionRef<void(FConcertSessionSerializedPayload&& Payload)> ProcessMoveable) override;
 		//~ End IReplicationDataSource Interface
 		
 		//~ Begin IReplicationCacheUser Interface
-		virtual void OnDataCached(const FConcertReplicatedObjectId& Object, TSharedRef<const FConcertReplication_ObjectReplicationEvent> Data) override;
+		virtual void OnDataCached(const FConcertReplicatedObjectId& Object, FSequenceId SequenceId, TSharedRef<const FConcertReplication_ObjectReplicationEvent> Data) override;
+		virtual void OnCachedDataUpdated(const FConcertReplicatedObjectId& Object, FSequenceId SequenceId) override;
 		//~ End IReplicationCacheUser Interface
 
 	protected:
@@ -41,8 +42,16 @@ namespace UE::ConcertSyncCore
 		
 	private:
 
+		struct FPendingObjectData
+		{
+			/** Pointer to what OnDataCached passed to us. */
+			TSharedPtr<const FConcertReplication_ObjectReplicationEvent> DataToApply;
+			/** The latest SequenceId that the data contains. The data might contain data from multiple sequences. */
+			FSequenceId SequenceId;
+		};
+
 		/** Stores events as they are received. */
-		TMap<FConcertReplicatedObjectId, TSharedPtr<const FConcertReplication_ObjectReplicationEvent>> PendingEvents;
+		TMap<FConcertReplicatedObjectId, FPendingObjectData> PendingEvents;
 
 		/** Provides us with replication events and shares them effectively. */
 		FObjectReplicationCache* ReplicationCache = nullptr;
diff --git a/Engine/Plugins/Developer/Concert/ConcertSync/ConcertSyncTest/Source/ConcertSyncTest/Private/Replication/Components/ReplicationCache.spec.cpp b/Engine/Plugins/Developer/Concert/ConcertSync/ConcertSyncTest/Source/ConcertSyncTest/Private/Replication/Components/ReplicationCache.spec.cpp
new file mode 100644
--- /dev/null
+++ b/Engine/Plugins/Developer/Concert/ConcertSync/ConcertSyncTest/Source/ConcertSyncTest/Private/Replication/Components/ReplicationCache.spec.cpp
@@ -0,0 +1,285 @@
+// Copyright Epic Games, Inc. All Rights Reserved.
+
+#include "Misc/AutomationTest.h"
+
+#include "Replication/Formats/IObjectReplicationFormat.h"
+#include "Replication/Messages/ObjectReplication.h"
+#include "Replication/Processing/ObjectReplicationCache.h"
+#include "Replication/Data/ObjectIds.h"
+#include "Replication/TestReflectionObject.h"
+#include "Util/ClientServerCommunicationTest.h"
+
+#include "Misc/AutomationTest.h"
+#include "Misc/Paths.h"
+#include "HAL/FileManager.h"
+
+namespace UE::ConcertSyncTests
+{
+	/** Fake format that just sends around FNativeStruct and combines the contained float value by addition; no real implementation would actually do that but this is easier for testing. */
+	class FTestReplicationFormat : public ConcertSyncCore::IObjectReplicationFormat
+	{
+	public:
+		virtual TOptional<FConcertSessionSerializedPayload> CreateReplicationEvent(UObject& Object, FAllowPropertyFunc IsPropertyAllowedFunc) override { return NotMocked<TOptional<FConcertSessionSerializedPayload>>({}); }
+		virtual void ClearInternalCache(TArrayView<UObject> ObjectsToClear) override { return NotMocked<void>(); }
+		virtual void CombineReplicationEvents(FConcertSessionSerializedPayload& Base, const FConcertSessionSerializedPayload& Newer) override
+		{
+			// Reuse FNativeStruct to avoid introducing more test types
+			FNativeStruct BaseStruct;
+			Base.GetTypedPayload(BaseStruct);
+			FNativeStruct NewerStruct;
+			Newer.GetTypedPayload(NewerStruct);
+
+			BaseStruct.Float += NewerStruct.Float;
+			Base.SetTypedPayload(BaseStruct);
+		}
+		virtual void ApplyReplicationEvent(UObject& Object, const FConcertSessionSerializedPayload& Payload) override { return NotMocked<void>(); }
+
+		static FConcertReplication_ObjectReplicationEvent CreateEvent(FSoftObjectPath Object, float Value)
+		{
+			FNativeStruct Data{ Value };
+			FConcertSessionSerializedPayload Payload;
+			Payload.SetTypedPayload(Data, EConcertPayloadCompressionType::None);
+			return { MoveTemp(Object), -1, Payload };
+		}
+	};
+
+	enum class EReplicationCacheTestFlags
+	{
+		NeverConsume = 0,
+		NeverReceive = 1,
+		ConsumeManually = 2
+	};
+	ENUM_CLASS_FLAGS(EReplicationCacheTestFlags);
+
+	/** Mock implementation which does different things based on EReplicationCacheTestFlags. */
+	class FTestReplicationCacheUser : public ConcertSyncCore::IReplicationCacheUser
+	{
+	public:
+
+		FAutomationTestBase& Test;
+		FConcertObjectInStreamID AllowedObject;
+		EReplicationCacheTestFlags Flags;
+		ConcertSyncCore::FSequenceId LastSequenceId = 100;
+		TSharedPtr<const FConcertReplication_ObjectReplicationEvent> CachedData;
+
+		bool bWasOnDataCachedCalled = false;
+		bool bWasOnCachedDataUpdatedCalled = false;
+
+		FTestReplicationCacheUser(FAutomationTestBase& Test, FConcertObjectInStreamID AllowedObject, EReplicationCacheTestFlags Flags)
+			: Test(Test)
+			, AllowedObject(MoveTemp(AllowedObject))
+			, Flags(Flags)
+		{}
+		
+		virtual bool WantsToAcceptObject(const FConcertReplicatedObjectId& Object) const override
+		{
+			return !EnumHasAnyFlags(Flags, EReplicationCacheTestFlags::NeverReceive) && AllowedObject == Object;
+		}
+		
+		virtual void OnDataCached(const FConcertReplicatedObjectId& Object, ConcertSyncCore::FSequenceId SequenceId, TSharedRef<const FConcertReplication_ObjectReplicationEvent> Data) override
+		{
+			bWasOnDataCachedCalled = true;
+			if (EnumHasAnyFlags(Flags, EReplicationCacheTestFlags::NeverReceive))
+			{
+				Test.AddError(TEXT("Received object that we never asked for!"));
+				return;
+			}
+			
+			LastSequenceId = SequenceId;
+			CachedData = Data;
+		}
+
+		virtual void OnCachedDataUpdated(const FConcertReplicatedObjectId& Object, ConcertSyncCore::FSequenceId SequenceId) override
+		{
+			bWasOnCachedDataUpdatedCalled = true;
+			LastSequenceId = SequenceId;
+		}
+
+		float PeakData() const
+		{
+			if (!CachedData)
+			{
+				return -1.f;
+			}
+
+			FNativeStruct Data;
+			CachedData->SerializedPayload.GetTypedPayload(Data);
+			return Data.Float;
+		}
+
+		const void* GetDataAddress() const
+		{
+			return CachedData.Get();
+		}
+		
+		void Consume()
+		{
+			check(!EnumHasAnyFlags(Flags, EReplicationCacheTestFlags::NeverConsume));
+			CachedData.Reset();
+		}
+
+		void ResetCallFlags()
+		{
+			bWasOnDataCachedCalled = false;
+			bWasOnCachedDataUpdatedCalled = false;
+		}
+	};
+	
+	BEGIN_DEFINE_SPEC(FReplicationCacheSpec, "Editor.Concert.Replication.Components.ReplicationCache", EAutomationTestFlags::EditorContext | EAutomationTestFlags::EngineFilter)
+		// Set up data
+		const FGuid StreamId = FGuid::NewGuid();
+		const FSoftObjectPath ObjectPath { TEXT("/Game/World.World:PersistentLevel.StaticMeshActor0") };
+		const FGuid DummySendingClientId = FGuid::NewGuid();
+		const FConcertReplicatedObjectId ObjectID{ { StreamId, ObjectPath }, DummySendingClientId};
+
+		// Events that will be sent ...
+		FConcertReplication_ObjectReplicationEvent Event_5;
+		FConcertReplication_ObjectReplicationEvent Event_10;
+		FConcertReplication_ObjectReplicationEvent Event_100;
+		// ... and their sequence IDs
+		const ConcertSyncCore::FSequenceId SequenceId_Event_5 = 0;
+		const ConcertSyncCore::FSequenceId SequenceId_Event_10 = 1;
+		const ConcertSyncCore::FSequenceId SequenceId_Event_100 = 2;
+
+		// That cache that will be tested
+		FTestReplicationFormat TestReplicationFormat;
+		TSharedPtr<ConcertSyncCore::FObjectReplicationCache> Cache;
+		
+		// A bunch of users that have different implementations
+		TSharedPtr<FTestReplicationCacheUser> User_NeverConsume;
+		TSharedPtr<FTestReplicationCacheUser> User_NeverReceive;
+		TSharedPtr<FTestReplicationCacheUser> User_ConsumeManually;
+	END_DEFINE_SPEC(FReplicationCacheSpec);
+
+	void FReplicationCacheSpec::Define()
+	{
+		BeforeEach([this]()
+		{
+			Event_5 = FTestReplicationFormat::CreateEvent(ObjectPath, 5.f);
+			Event_10 = FTestReplicationFormat::CreateEvent(ObjectPath, 10.f);
+			Event_100 = FTestReplicationFormat::CreateEvent(ObjectPath, 100.f);
+			
+			Cache					= MakeShared<ConcertSyncCore::FObjectReplicationCache>(TestReplicationFormat);
+			
+			User_NeverConsume		= MakeShared<FTestReplicationCacheUser>(*this, ObjectID, EReplicationCacheTestFlags::NeverConsume);
+			User_NeverReceive		= MakeShared<FTestReplicationCacheUser>(*this, ObjectID, EReplicationCacheTestFlags::NeverReceive);
+			User_ConsumeManually	= MakeShared<FTestReplicationCacheUser>(*this, ObjectID, EReplicationCacheTestFlags::ConsumeManually);
+
+			Cache->RegisterDataCacheUser(User_NeverConsume.ToSharedRef());
+			Cache->RegisterDataCacheUser(User_NeverReceive.ToSharedRef());
+			Cache->RegisterDataCacheUser(User_ConsumeManually.ToSharedRef());
+		});
+
+		Describe("StoreUntilConsumed", [this]()
+		{
+			It("Calls OnDataCached if WantsToAcceptObject", [this]()
+			{
+				// 1st event
+				Cache->StoreUntilConsumed(DummySendingClientId, StreamId, SequenceId_Event_5, Event_5);
+				
+				// User_NeverConsume
+				TestEqual(TEXT("User_NeverConsume: 5"), User_NeverConsume->PeakData(), 5.f);
+				TestFalse(TEXT("User_NeverConsume: Didn't call OnCachedDataUpdated()"), User_NeverConsume->bWasOnCachedDataUpdatedCalled);
+				TestTrue(TEXT("User_NeverConsume: Called OnDataCached()"), User_NeverConsume->bWasOnDataCachedCalled);
+				TestEqual(TEXT("Sequence number: never consume"), User_NeverConsume->LastSequenceId, SequenceId_Event_5);
+
+				// User_ConsumeManually
+				TestEqual(TEXT("User_ConsumeManually: 5"), User_ConsumeManually->PeakData(), 5.f);
+				TestFalse(TEXT("User_ConsumeManually: Didn't call OnCachedDataUpdated()"), User_ConsumeManually->bWasOnCachedDataUpdatedCalled);
+				TestTrue(TEXT("User_ConsumeManually: Called OnDataCached()"), User_ConsumeManually->bWasOnDataCachedCalled);
+				TestEqual(TEXT("Sequence number: manual consume"), User_ConsumeManually->LastSequenceId, SequenceId_Event_5);
+
+				// User_NeverReceive
+				TestFalse(TEXT("User_NeverReceive->bWasOnDataCachedCalled == false"), User_NeverReceive->bWasOnDataCachedCalled);
+				TestFalse(TEXT("User_NeverReceive->bWasOnCachedDataUpdatedCalled == false"), User_NeverReceive->bWasOnCachedDataUpdatedCalled);
+			});
+
+			It("Combines data", [this]()
+			{
+				// 1st event
+				Cache->StoreUntilConsumed(DummySendingClientId, StreamId, SequenceId_Event_5, Event_5);
+				// 2nd event
+				User_NeverConsume->ResetCallFlags();
+				User_ConsumeManually->ResetCallFlags();
+				Cache->StoreUntilConsumed(DummySendingClientId, StreamId, SequenceId_Event_10, Event_10);
+
+				// User_NeverConsume
+				TestEqual(TEXT("Combined events: 5 and 10"), User_NeverConsume->PeakData(), 15.f);
+				TestTrue(TEXT("User_NeverConsume: Called OnCachedDataUpdated()"), User_NeverConsume->bWasOnCachedDataUpdatedCalled);
+				TestFalse(TEXT("User_NeverConsume: Didn't call OnDataCached()"), User_NeverConsume->bWasOnDataCachedCalled);
+				TestEqual(TEXT("Sequence number: never consume"), User_NeverConsume->LastSequenceId, SequenceId_Event_10);
+			});
+
+			It("Calls OnDataCached again if data was already consumed", [this]()
+			{
+				// 1st event
+				Cache->StoreUntilConsumed(DummySendingClientId, StreamId, SequenceId_Event_5, Event_5);
+
+				TestEqual(TEXT("Received new data: 5"), User_ConsumeManually->PeakData(), 5.f);
+				TestFalse(TEXT("User_ConsumeManually: Didn't call OnCachedDataUpdated()"), User_ConsumeManually->bWasOnCachedDataUpdatedCalled);
+				TestTrue(TEXT("User_ConsumeManually: Called OnDataCached()"), User_ConsumeManually->bWasOnDataCachedCalled);
+				TestEqual(TEXT("Sequence number: 5"), User_ConsumeManually->LastSequenceId, SequenceId_Event_5);
+				
+				// 2nd event
+				User_ConsumeManually->Consume();
+				User_ConsumeManually->ResetCallFlags();
+				Cache->StoreUntilConsumed(DummySendingClientId, StreamId, SequenceId_Event_10, Event_10);
+
+				// User_ConsumeManually
+				TestEqual(TEXT("Received new data: 10"), User_ConsumeManually->PeakData(), 10.f);
+				TestFalse(TEXT("User_ConsumeManually: Didn't call OnCachedDataUpdated() again"), User_ConsumeManually->bWasOnCachedDataUpdatedCalled);
+				TestTrue(TEXT("User_ConsumeManually: Called OnDataCached() again"), User_ConsumeManually->bWasOnDataCachedCalled);
+				TestEqual(TEXT("Sequence number: 10"), User_ConsumeManually->LastSequenceId, SequenceId_Event_10);
+			});
+
+			It("Calls OnCachedDataUpdated if data was not yet consumed", [this]()
+			{
+				// 1st event
+				Cache->StoreUntilConsumed(DummySendingClientId, StreamId, SequenceId_Event_5, Event_5);
+
+				TestEqual(TEXT("Received new data: 5"), User_ConsumeManually->PeakData(), 5.f);
+				TestFalse(TEXT("User_ConsumeManually: Didn't call OnCachedDataUpdated()"), User_ConsumeManually->bWasOnCachedDataUpdatedCalled);
+				TestTrue(TEXT("User_ConsumeManually: Called OnDataCached()"), User_ConsumeManually->bWasOnDataCachedCalled);
+				TestEqual(TEXT("Sequence number: 5"), User_ConsumeManually->LastSequenceId, SequenceId_Event_5);
+				
+				// 2nd event
+				User_ConsumeManually->ResetCallFlags();
+				Cache->StoreUntilConsumed(DummySendingClientId, StreamId, SequenceId_Event_10, Event_10);
+
+				// User_ConsumeManually
+				TestEqual(TEXT("Received new data: 15"), User_ConsumeManually->PeakData(), 15.f);
+				TestTrue(TEXT("User_ConsumeManually: Called OnCachedDataUpdated()"), User_ConsumeManually->bWasOnCachedDataUpdatedCalled);
+				TestFalse(TEXT("User_ConsumeManually: Didn't call OnDataCached()"), User_ConsumeManually->bWasOnDataCachedCalled);
+				TestEqual(TEXT("Sequence number: 10"), User_ConsumeManually->LastSequenceId, SequenceId_Event_10);
+			});
+		});
+
+		Describe("Cache does not leak data", [this]()
+		{
+			It("When unregistering users", [this]()
+			{
+				// 1st event
+				Cache->StoreUntilConsumed(DummySendingClientId, StreamId, SequenceId_Event_5, Event_5);
+				User_ConsumeManually->Consume();
+				
+				// 2nd event
+				Cache->StoreUntilConsumed(DummySendingClientId, StreamId, SequenceId_Event_10, Event_10);
+				Cache->UnregisterDataCacheUser(User_ConsumeManually.ToSharedRef());
+				Cache->RegisterDataCacheUser(User_ConsumeManually.ToSharedRef());
+				const void* AddressBefore = User_ConsumeManually->GetDataAddress();
+				
+				// 3rd event
+				User_NeverConsume->ResetCallFlags();
+				User_ConsumeManually->ResetCallFlags();
+				Cache->StoreUntilConsumed(DummySendingClientId, StreamId, SequenceId_Event_100, Event_100);
+				
+				// User_ConsumeManually
+				TestEqual(TEXT("Re-registered user only has the new data"), User_ConsumeManually->PeakData(), 100.f);
+				TestNotEqual(TEXT("Re-registered user's data was allocated in different memory block"), AddressBefore, User_ConsumeManually->GetDataAddress());
+				TestFalse(TEXT("User_ConsumeManually: Didn't call OnCachedDataUpdated()"), User_ConsumeManually->bWasOnCachedDataUpdatedCalled);
+				TestTrue(TEXT("User_ConsumeManually: Called OnDataCached()"), User_ConsumeManually->bWasOnDataCachedCalled);
+				TestEqual(TEXT("Sequence number: manual consume"), User_ConsumeManually->LastSequenceId, SequenceId_Event_100);
+			});
+		});
+	}
+}
diff --git a/Engine/Plugins/Developer/Concert/ConcertSync/ConcertSyncTest/Source/ConcertSyncTest/Private/Replication/Components/ReplicationCacheTests.cpp b/Engine/Plugins/Developer/Concert/ConcertSync/ConcertSyncTest/Source/ConcertSyncTest/Private/Replication/Components/ReplicationCacheTests.cpp
deleted file mode 100644
