Fix typo in operator==; it compared Left == Left instead of Left == Right.

diff --git a/Engine/Plugins/Developer/Concert/ConcertInsights/ConcertInsightsVisualizer/Source/ConcertInsightsVisualizer/Private/Trace/Types/ScopeInfos.h b/Engine/Plugins/Developer/Concert/ConcertInsights/ConcertInsightsVisualizer/Source/ConcertInsightsVisualizer/Private/Trace/Types/ScopeInfos.h
--- a/Engine/Plugins/Developer/Concert/ConcertInsights/ConcertInsightsVisualizer/Source/ConcertInsightsVisualizer/Private/Trace/Types/ScopeInfos.h
+++ b/Engine/Plugins/Developer/Concert/ConcertInsights/ConcertInsightsVisualizer/Source/ConcertInsightsVisualizer/Private/Trace/Types/ScopeInfos.h
@@ -44,7 +44,7 @@ namespace UE::ConcertInsightsVisualizer
 
 		struct FEndpointScopeInfo MakeEndpointInfo(FEndpointId EndpointId) const;
 
-		friend bool operator==(const FSequenceScopeInfo& Left, const FSequenceScopeInfo& Right) { return static_cast<const FObjectScopeInfo&>(Left) == static_cast<const FObjectScopeInfo&>(Left) && Left.SequenceId == Right.SequenceId; }
+		friend bool operator==(const FSequenceScopeInfo& Left, const FSequenceScopeInfo& Right) { return static_cast<const FObjectScopeInfo&>(Left) == static_cast<const FObjectScopeInfo&>(Right) && Left.SequenceId == Right.SequenceId; }
 		friend bool operator!=(const FSequenceScopeInfo& Left, const FSequenceScopeInfo& Right) { return !(Left == Right); }
 		friend uint32 GetTypeHash(const FSequenceScopeInfo& Info) { return HashCombine(GetTypeHash(Info.SequenceId), GetTypeHash(static_cast<const FObjectScopeInfo&>(Info))); }
 	};
