Fixing single header compilation.

diff --git a/Engine/Plugins/Developer/Concert/ConcertInsights/ConcertInsightsVisualizer/Source/ConcertInsightsVisualizer/Private/Trace/Types/SequenceEndpointPair.h b/Engine/Plugins/Developer/Concert/ConcertInsights/ConcertInsightsVisualizer/Source/ConcertInsightsVisualizer/Private/Trace/Types/SequenceEndpointPair.h
--- a/Engine/Plugins/Developer/Concert/ConcertInsights/ConcertInsightsVisualizer/Source/ConcertInsightsVisualizer/Private/Trace/Types/SequenceEndpointPair.h
+++ b/Engine/Plugins/Developer/Concert/ConcertInsights/ConcertInsightsVisualizer/Source/ConcertInsightsVisualizer/Private/Trace/Types/SequenceEndpointPair.h
@@ -4,6 +4,7 @@
 
 #include "EndpointId.h"
 #include "SequenceId.h"
+#include "Templates/TypeHash.h"
 
 namespace UE::ConcertInsightsVisualizer
 {
diff --git a/Engine/Plugins/Developer/Concert/ConcertMain/Source/ConcertTransport/Public/Trace/Private/ConcertReplicationSequenceContext.h b/Engine/Plugins/Developer/Concert/ConcertMain/Source/ConcertTransport/Public/Trace/Private/ConcertReplicationSequenceContext.h
--- a/Engine/Plugins/Developer/Concert/ConcertMain/Source/ConcertTransport/Public/Trace/Private/ConcertReplicationSequenceContext.h
+++ b/Engine/Plugins/Developer/Concert/ConcertMain/Source/ConcertTransport/Public/Trace/Private/ConcertReplicationSequenceContext.h
@@ -2,6 +2,9 @@
 
 #pragma once
 
+#include "Containers/UnrealString.h"
+#include "UObject/SoftObjectPath.h"
+
 namespace UE::ConcertTrace::Private
 {
 	struct FConcertReplicationSequenceContext
