Fix slowing down Unreal Insights due to ConcertInsights aggregating on the game thread.
Insights devs may locally enable the plugin to iterate on API but not actually use it at runtime.
So this introduces a CVar to disable aggregation altogether by default.

diff --git a/Engine/Plugins/Developer/Concert/ConcertInsights/ConcertInsightsVisualizer/Source/ConcertInsightsVisualizer/Private/Trace/ProtocolMultiEndpointProvider.cpp b/Engine/Plugins/Developer/Concert/ConcertInsights/ConcertInsightsVisualizer/Source/ConcertInsightsVisualizer/Private/Trace/ProtocolMultiEndpointProvider.cpp
--- a/Engine/Plugins/Developer/Concert/ConcertInsights/ConcertInsightsVisualizer/Source/ConcertInsightsVisualizer/Private/Trace/ProtocolMultiEndpointProvider.cpp
+++ b/Engine/Plugins/Developer/Concert/ConcertInsights/ConcertInsightsVisualizer/Source/ConcertInsightsVisualizer/Private/Trace/ProtocolMultiEndpointProvider.cpp
@@ -5,6 +5,7 @@
 #include "LogConcertInsights.h"
 
 #include "Algo/AnyOf.h"
+#include "HAL/IConsoleManager.h"
 #include "Insights/IUnrealInsightsModule.h"
 #include "Modules/ModuleManager.h"
 #include "Util/TimeSyncUtils.h"
@@ -34,6 +35,35 @@ namespace UE::ConcertInsightsVisualizer
 				*Init.GetTraceInitTimeUtc().ToString(TEXT("%Y-%m-%d_%H-%Mm-%Ss-%sms"))
 				);
 		}
+
+		// Read SubscribeTickerIfEnabled for more info.
+		static TAutoConsoleVariable<bool> CVarEnableGameThreadAggregation(
+			TEXT("Insights.Concert.EnableGameThreadAggregation"),
+			false,
+			TEXT("Whether aggregation of the trace files should occur on the game thread (can freeze UI)."),
+			ECVF_Default
+		);
+
+		template<typename TCreateDelegateLambda>
+		static FTSTicker::FDelegateHandle SubscribeTickerIfEnabled(TCreateDelegateLambda&& CreateDelegateLambda)
+		{
+			if (CVarEnableGameThreadAggregation.GetValueOnAnyThread())
+			{
+				return FTSTicker::GetCoreTicker().AddTicker(CreateDelegateLambda());
+			}
+
+			// This is an experimental plugin.
+			// Aggregation currently is only implemented to occur on the game thread, which can freeze the program.
+			// This plugin may get enabled but not used (e.g. Insights devs iterating on the API may enable this plugin to check it compiles, etc. but not actually use it at runtime).
+			// Do not slow down the program for such cases - so by default this CVar is disabled.
+			// Since this plugin is experimental and not being actively worked on, we're not going to implement performant aggregation right now.
+			// The proper solution would be as follows:
+			// 1. FProtocolMultiEndpointProvider::ProcessGeneratedTraceData should occur on a separate thread because it takes long
+			// 2. That new thread and the game thread should be synchronized with custom read / write locks, we'd have to introduce, as well.
+			UE_LOG(LogConcertInsights, Warning, TEXT("ConcertInsights will not work because console variable Insights.Concert.EnableGameThreadAggregation is set to false."));
+			return FTSTicker::FDelegateHandle{};
+		}
+		
 	}
 	
 	const FName FProtocolMultiEndpointProvider::ProviderName = TEXT("FProtocolMultiEndpointProvider");
@@ -41,12 +71,15 @@ namespace UE::ConcertInsightsVisualizer
 	FProtocolMultiEndpointProvider::FProtocolMultiEndpointProvider(TraceServices::IAnalysisSession& Session UE_LIFETIMEBOUND)
 		: Session(Session)
 		, Aggregator(*FModuleManager::LoadModuleChecked<IUnrealInsightsModule>("TraceInsights").GetStoreClient(), GetMainTraceId())
-		, TickHandle(FTSTicker::GetCoreTicker().AddTicker(FTickerDelegate::CreateRaw(this, &FProtocolMultiEndpointProvider::ProcessGeneratedTraceData)))
+		, TickHandle(Private::SubscribeTickerIfEnabled([this]{ return FTickerDelegate::CreateRaw(this, &FProtocolMultiEndpointProvider::ProcessGeneratedTraceData); }))
 	{}
 
 	FProtocolMultiEndpointProvider::~FProtocolMultiEndpointProvider()
 	{
-		FTSTicker::GetCoreTicker().RemoveTicker(TickHandle);
+		if (TickHandle.IsValid())
+		{
+			FTSTicker::GetCoreTicker().RemoveTicker(TickHandle);
+		}
 	}
 
 	const TCHAR* FProtocolMultiEndpointProvider::GetEndpointDisplayName(FEndpointId EndpointId) const
diff --git a/Engine/Plugins/Developer/Concert/ConcertInsights/README.txt b/Engine/Plugins/Developer/Concert/ConcertInsights/README.txt
--- a/Engine/Plugins/Developer/Concert/ConcertInsights/README.txt
+++ b/Engine/Plugins/Developer/Concert/ConcertInsights/README.txt
@@ -27,6 +27,7 @@ To try this out, the minimum set-up is:
 - Either unshelf 33391924, or recreate the work in that CL:
 	- Engine/Programs/UnrealMultiUserSlateServerConfig/DefaultEngine.ini: Add +ProgramEnabledPlugins=ConcertInsightsServer
 	- Engine/Programs/UnrealInsights/DefaultEngine.ini: Add +ProgramEnabledPlugins=ConcertInsightsVisualizer
+- Launch Insights with the console variable "Insights.Concert.EnableGameThreadAggregation = true".
 	
 ========== Tutorial with VCam ==========
 To view example tracing, do the following:
